Trace Source filtering plugin, used to narrow down (filter) which AActors, can, output Gameplay and Animation trace data. Currently only supported in Editor.

- Adds new Source Filtering Tab inside of Gameplay Insights (Editor), allows for setting up filter (sets) and filtering UWorlds
- Introduces definition of DataSourceFilter(s), native of blueprint implemented interface for determining whether or not an individual Actor is allowed to ouput data
- Per-world manager that applies the filters to its represented UWorld during runtime
- Filtering specific Trace events, future-proof for looping back filtering state to standalone UnrealInsights
- Enabled plugin by default in QAGame

diff --git a/Engine/Plugins/Developer/TraceSourceFiltering/Config/BaseTraceSourceFilters.ini b/Engine/Plugins/Developer/TraceSourceFiltering/Config/BaseTraceSourceFilters.ini
new file mode 100644
--- /dev/null
+++ b/Engine/Plugins/Developer/TraceSourceFiltering/Config/BaseTraceSourceFilters.ini
@@ -0,0 +1,2 @@
+[WorldTypeTraceFilter]
+DisabledWorldTypes=2+4
\ No newline at end of file
diff --git a/Engine/Plugins/Developer/TraceSourceFiltering/Source/SourceFilteringCore/Private/DataSourceFiltering.cpp b/Engine/Plugins/Developer/TraceSourceFiltering/Source/SourceFilteringCore/Private/DataSourceFiltering.cpp
new file mode 100644
--- /dev/null
+++ b/Engine/Plugins/Developer/TraceSourceFiltering/Source/SourceFilteringCore/Private/DataSourceFiltering.cpp
@@ -0,0 +1,7 @@
+// Copyright Epic Games, Inc. All Rights Reserved.
+
+#include "DataSourceFiltering.h"
+#include "CoreTypes.h"
+#include "Modules/ModuleManager.h"
+
+IMPLEMENT_MODULE(FDefaultModuleImpl, SourceFilteringCore);
\ No newline at end of file
diff --git a/Engine/Plugins/Developer/TraceSourceFiltering/Source/SourceFilteringCore/Public/DataSourceFiltering.h b/Engine/Plugins/Developer/TraceSourceFiltering/Source/SourceFilteringCore/Public/DataSourceFiltering.h
new file mode 100644
--- /dev/null
+++ b/Engine/Plugins/Developer/TraceSourceFiltering/Source/SourceFilteringCore/Public/DataSourceFiltering.h
@@ -0,0 +1,38 @@
+// Copyright Epic Games, Inc. All Rights Reserved.
+
+#pragma once
+
+#include "CoreTypes.h"
+#include "Misc/EnumRange.h"
+#include "UObject/ObjectMacros.h"
+
+/** Enum representing the possible Filter Set operations */
+UENUM()
+enum class EFilterSetMode : uint8
+{
+	AND,
+	OR,
+	NOT,
+	Count UMETA(Hidden)
+};
+
+ENUM_RANGE_BY_COUNT(EFilterSetMode, EFilterSetMode::Count);
+
+/** Enum representing all possible Source Filter operations, used by FSourceFilterTrace::OutputFilterOperation */
+enum class ESourceActorFilterOperation
+{
+	RemoveFilter,
+	MoveFilter,
+	ReplaceFilter,
+	SetFilterMode,
+	SetFilterState
+};
+
+/** Enum representing all possible World Filter operations, used by FSourceFilterTrace::OutputWorldOperation */
+enum class EWorldFilterOperation
+{
+	TypeFilter,
+	NetModeFilter,
+	InstanceFilter,
+	RemoveWorld
+};
diff --git a/Engine/Plugins/Developer/TraceSourceFiltering/Source/SourceFilteringCore/Public/IDataSourceFilterInterface.h b/Engine/Plugins/Developer/TraceSourceFiltering/Source/SourceFilteringCore/Public/IDataSourceFilterInterface.h
new file mode 100644
--- /dev/null
+++ b/Engine/Plugins/Developer/TraceSourceFiltering/Source/SourceFilteringCore/Public/IDataSourceFilterInterface.h
@@ -0,0 +1,41 @@
+// Copyright Epic Games, Inc. All Rights Reserved.
+
+#pragma once
+
+#include "CoreTypes.h"
+#include "UObject/Interface.h"
+#include "Internationalization/Internationalization.h"
+#include "IDataSourceFilterInterface.generated.h"
+
+UINTERFACE(Blueprintable)
+class SOURCEFILTERINGCORE_API UDataSourceFilterInterface : public UInterface
+{
+	GENERATED_BODY()
+};
+
+/** Interface used for implementing Engine and UnrealInsights versions respectively UDataSourceFilter and UTraceDataSourceFilter */
+class SOURCEFILTERINGCORE_API IDataSourceFilterInterface
+{
+	GENERATED_BODY()
+
+public:
+	UFUNCTION(BlueprintNativeEvent, Category = TraceSourceFiltering)
+	void GetDisplayText(FText& OutDisplayText) const;
+	void GetDisplayText_Implementation(FText& OutDisplayText) const
+	{
+		return GetDisplayText_Internal(OutDisplayText);
+	}
+
+	UFUNCTION(BlueprintNativeEvent, Category = TraceSourceFiltering)
+	void GetToolTipText(FText& OutDisplayText) const;
+	void GetToolTipText_Implementation(FText& OutDisplayText) const
+	{
+		return GetToolTipText_Internal(OutDisplayText);
+	}
+
+	virtual void SetEnabled(bool bState) = 0;
+	virtual bool IsEnabled() const = 0;
+protected:
+	virtual void GetDisplayText_Internal(FText& OutDisplayText) const { OutDisplayText = NSLOCTEXT("IDataSourceFilterInterface", "GetDisplayTextNotImplementedText", "IDataSourceFilterInterface::GetDisplayText_Internal not overidden (missing implementation)"); }
+	virtual void GetToolTipText_Internal(FText& OutDisplayText) const { GetDisplayText_Internal(OutDisplayText); }
+};
\ No newline at end of file
diff --git a/Engine/Plugins/Developer/TraceSourceFiltering/Source/SourceFilteringCore/Public/IDataSourceFilterSetInterface.h b/Engine/Plugins/Developer/TraceSourceFiltering/Source/SourceFilteringCore/Public/IDataSourceFilterSetInterface.h
new file mode 100644
--- /dev/null
+++ b/Engine/Plugins/Developer/TraceSourceFiltering/Source/SourceFilteringCore/Public/IDataSourceFilterSetInterface.h
@@ -0,0 +1,25 @@
+// Copyright Epic Games, Inc. All Rights Reserved.
+
+#pragma once
+
+#include "CoreTypes.h"
+#include "IDataSourceFilterInterface.h"
+#include "DataSourceFiltering.h"
+#include "IDataSourceFilterSetInterface.generated.h"
+
+UINTERFACE(Blueprintable)
+class SOURCEFILTERINGCORE_API UDataSourceFilterSetInterface : public UInterface
+{
+	GENERATED_BODY()
+};
+
+/** Interface used for implementing Engine and UnrealInsights versions respectively UDataSourceFilterSet and UTraceDataSourceFilterSet */
+class SOURCEFILTERINGCORE_API IDataSourceFilterSetInterface
+{
+	GENERATED_BODY()
+
+public:
+	/** Return the current Filter Set operation mode */
+	virtual EFilterSetMode GetFilterSetMode() const = 0;
+};
+
diff --git a/Engine/Plugins/Developer/TraceSourceFiltering/Source/SourceFilteringCore/Public/TraceSourceFilteringSettings.h b/Engine/Plugins/Developer/TraceSourceFiltering/Source/SourceFilteringCore/Public/TraceSourceFilteringSettings.h
new file mode 100644
--- /dev/null
+++ b/Engine/Plugins/Developer/TraceSourceFiltering/Source/SourceFilteringCore/Public/TraceSourceFilteringSettings.h
@@ -0,0 +1,29 @@
+// Copyright Epic Games, Inc. All Rights Reserved.
+
+#pragma once
+
+#include "CoreTypes.h"
+#include "UObject/ObjectMacros.h"
+
+#include "TraceSourceFilteringSettings.generated.h"
+
+UCLASS(Config=TraceSourceFilters)
+class SOURCEFILTERINGCORE_API UTraceSourceFilteringSettings : public UObject
+{
+	GENERATED_BODY()
+
+public:
+	UTraceSourceFilteringSettings() : bDrawFilteringStates(false), bDrawOnlyPassingActors(false), bDrawFilterDescriptionForRejectedActors(false) {}
+
+	/** Whether or not the filtering state for all considered AActor's inside for a UWorld should be drawn using a wire frame box */
+	UPROPERTY(Config)
+	bool bDrawFilteringStates;
+
+	/** Whether or not only AActor's that are not filtered out should be considered for drawing their wireframe box */
+	UPROPERTY(Config)
+	bool bDrawOnlyPassingActors;
+
+	/** Whether or not to draw the failed UDataSourceFilter's description for AActor's that did not pass the filtering */
+	UPROPERTY(Config)
+	bool bDrawFilterDescriptionForRejectedActors;
+};
\ No newline at end of file
diff --git a/Engine/Plugins/Developer/TraceSourceFiltering/Source/SourceFilteringCore/SourceFilteringCore.Build.cs b/Engine/Plugins/Developer/TraceSourceFiltering/Source/SourceFilteringCore/SourceFilteringCore.Build.cs
new file mode 100644
--- /dev/null
+++ b/Engine/Plugins/Developer/TraceSourceFiltering/Source/SourceFilteringCore/SourceFilteringCore.Build.cs
@@ -0,0 +1,16 @@
+// Copyright Epic Games, Inc. All Rights Reserved.
+
+using UnrealBuildTool;
+
+public class SourceFilteringCore : ModuleRules
+{
+    public SourceFilteringCore(ReadOnlyTargetRules Target) : base(Target)
+    {
+        PrivateDependencyModuleNames.AddRange(
+                new string[] {
+                    "Core",
+                    "CoreUObject",
+                }
+            );
+    }
+}
diff --git a/Engine/Plugins/Developer/TraceSourceFiltering/Source/SourceFilteringEditor/Private/EditorSessionSourceFilterService.cpp b/Engine/Plugins/Developer/TraceSourceFiltering/Source/SourceFilteringEditor/Private/EditorSessionSourceFilterService.cpp
new file mode 100644
--- /dev/null
+++ b/Engine/Plugins/Developer/TraceSourceFiltering/Source/SourceFilteringEditor/Private/EditorSessionSourceFilterService.cpp
@@ -0,0 +1,525 @@
+// Copyright Epic Games, Inc. All Rights Reserved.
+
+#if WITH_EDITOR
+
+#include "EditorSessionSourceFilterService.h"
+
+#include "UObject/UObjectGlobals.h"
+#include "UObject/SoftObjectPath.h"
+#include "ScopedTransaction.h"
+#include "Editor.h"
+
+#include "TraceSourceFiltering.h"
+#include "IFilterObject.h"
+#include "FilterObject.h"
+#include "FilterSetObject.h"
+
+#include "ClassViewerFilter.h"
+#include "ClassViewerModule.h"
+#include "DataSourceFilter.h"
+#include "DataSourceFilterSet.h"
+
+#include "PackageHelperFunctions.h"
+#include "IContentBrowserSingleton.h"
+#include "ContentBrowserModule.h"
+#include "Misc/FileHelper.h"
+#include "AssetRegistryModule.h"
+#include "EditorDirectories.h"
+#include "AssetToolsModule.h"
+#include "TraceWorldFiltering.h"
+#include "WorldObject.h"
+#include "WorldFilters.h"
+#include "SourceFilterStyle.h"
+#include "TraceFilter.h"
+
+#define LOCTEXT_NAMESPACE "FEditorSourceFilterService"
+
+FString FEditorSessionSourceFilterService::TransactionContext = TEXT("FEditorSourceFilterService");
+
+FEditorSessionSourceFilterService::FEditorSessionSourceFilterService()
+{
+	FilterCollection = FTraceSourceFiltering::Get().GetFilterCollection();
+	if (GEditor)
+	{
+		GEditor->RegisterForUndo(this);
+	}
+
+	// Register delegate to catch engine-level trace filtering system changes 
+	FTraceWorldFiltering::OnFilterStateChanged().AddRaw(this, &FEditorSessionSourceFilterService::UpdateTimeStamp);	
+	SetupWorldFilters();
+}
+
+FEditorSessionSourceFilterService::~FEditorSessionSourceFilterService()
+{
+	if (GEditor)
+	{
+		GEditor->UnregisterForUndo(this);
+	}
+	
+	FTraceWorldFiltering::OnFilterStateChanged().RemoveAll(this);
+}
+
+void FEditorSessionSourceFilterService::AddFilter(const FString& FilterClassName)
+{
+	FSoftClassPath ClassPath(FilterClassName);
+	if (UClass* Class = ClassPath.TryLoadClass<UDataSourceFilter>())
+	{
+		const FScopedTransaction Transaction(*FEditorSessionSourceFilterService::TransactionContext, LOCTEXT("AddFilter", "Add Filter"), FilterCollection);
+		FilterCollection->Modify();
+
+		FilterCollection->AddFilterOfClass(Class);
+
+		UpdateTimeStamp();
+	}
+}
+
+void FEditorSessionSourceFilterService::AddFilterToSet(TSharedRef<const IFilterObject> FilterSet, const FString& FilterClassName)
+{
+	FSoftClassPath ClassPath(FilterClassName);
+	if (UClass* Class = ClassPath.TryLoadClass<UDataSourceFilter>())
+	{
+		const FScopedTransaction Transaction(*FEditorSessionSourceFilterService::TransactionContext, LOCTEXT("AddFilter", "Add Filter to Filter Set"), FilterCollection);
+		FilterCollection->Modify();
+
+		FilterCollection->AddFilterOfClassToSet(Class, CastChecked<UDataSourceFilterSet>(FilterSet->GetFilter()));
+
+		UpdateTimeStamp();
+	}
+}
+
+void FEditorSessionSourceFilterService::AddFilterToSet(TSharedRef<const IFilterObject> FilterSet, TSharedRef<const IFilterObject> ExistingFilter)
+{
+	const FScopedTransaction Transaction(*FEditorSessionSourceFilterService::TransactionContext, LOCTEXT("MoveFilter", "Moving Filter"), FilterCollection);
+	FilterCollection->Modify();
+
+	UDataSourceFilter* Filter = CastChecked<UDataSourceFilter>(ExistingFilter->GetFilter());
+	UDataSourceFilterSet* Set = CastChecked<UDataSourceFilterSet>(FilterSet->GetFilter());
+	FilterCollection->MoveFilter(Filter, Set);
+
+	UpdateTimeStamp();
+}
+
+void FEditorSessionSourceFilterService::RemoveFilter(TSharedRef<const IFilterObject> InFilter)
+{
+	const FScopedTransaction Transaction(*FEditorSessionSourceFilterService::TransactionContext, LOCTEXT("RemoveFilter", "Removing Filter"), InFilter->GetFilter());
+	FilterCollection->Modify();
+
+	FilterCollection->RemoveFilter(CastChecked<UDataSourceFilter>(InFilter->GetFilter()));
+
+	UpdateTimeStamp();
+}
+
+void FEditorSessionSourceFilterService::SetFilterSetMode(TSharedRef<const IFilterObject> InFilter, EFilterSetMode Mode)
+{
+	if (UDataSourceFilterSet* FilterSet = CastChecked<UDataSourceFilterSet>(InFilter->GetFilter()))
+	{
+		const FScopedTransaction Transaction(*FEditorSessionSourceFilterService::TransactionContext, LOCTEXT("SetFilterSetMode", "Set Filter Set Mode"), FilterSet);
+		FilterSet->Modify();
+
+		FilterSet->SetFilterMode(Mode);
+
+		UpdateTimeStamp();
+	}
+}
+
+void FEditorSessionSourceFilterService::SetFilterState(TSharedRef<const IFilterObject> InFilter, bool bState)
+{
+	if (UDataSourceFilter* Filter = CastChecked<UDataSourceFilter>(InFilter->GetFilter()))
+	{
+		const FScopedTransaction Transaction(*FEditorSessionSourceFilterService::TransactionContext, LOCTEXT("SetFilterState", "Set Filter State"), Filter);
+		Filter->Modify();
+
+		Filter->SetEnabled(bState);
+	}
+}
+
+void FEditorSessionSourceFilterService::ResetFilters()
+{
+	const FScopedTransaction Transaction(*FEditorSessionSourceFilterService::TransactionContext, LOCTEXT("ResetFilters", "Reset Filters"), FilterCollection);
+	FilterCollection->Modify();
+
+	FilterCollection->Reset();
+
+	UpdateTimeStamp();
+}
+
+void FEditorSessionSourceFilterService::UpdateTimeStamp()
+{
+	Timestamp = FDateTime::Now();
+}
+
+UTraceSourceFilteringSettings* FEditorSessionSourceFilterService::GetFilterSettings()
+{
+	return FTraceSourceFiltering::Get().GetSettings();
+}
+
+bool FEditorSessionSourceFilterService::IsActionPending() const
+{
+	return false;
+}
+
+TSharedRef<SWidget> FEditorSessionSourceFilterService::GetFilterPickerWidget(FOnFilterClassPicked InFilterClassPicked)
+{
+	/** Class filter implementation, ensuring we only show valid UDataSourceFilter (sub)classes */
+	class FFilterClassFilter : public IClassViewerFilter
+	{
+	public:
+		bool IsClassAllowed(const FClassViewerInitializationOptions& InInitOptions, const UClass* InClass, TSharedRef< FClassViewerFilterFuncs > InFilterFuncs) override
+		{
+			if (InClass->HasAnyClassFlags(CLASS_Abstract | CLASS_HideDropDown | CLASS_Deprecated))
+			{
+				return false;
+			}
+
+			return InClass->IsChildOf(UDataSourceFilter::StaticClass()) && InClass != UDataSourceFilter::StaticClass() && InClass != UDataSourceFilterSet::StaticClass();
+		}
+
+		virtual bool IsUnloadedClassAllowed(const FClassViewerInitializationOptions& InInitOptions, const TSharedRef< const IUnloadedBlueprintData > InClass, TSharedRef< FClassViewerFilterFuncs > InFilterFuncs) override
+		{
+			return InClass->IsChildOf(UDataSourceFilter::StaticClass());
+		}
+	};
+
+	FClassViewerInitializationOptions Options;
+	Options.bShowUnloadedBlueprints = true;
+	Options.bShowNoneOption = false;
+	TSharedPtr<FFilterClassFilter> ClassFilter = MakeShareable(new FFilterClassFilter);
+	Options.ClassFilter = ClassFilter;
+
+	FOnClassPicked ClassPicked = FOnClassPicked::CreateLambda([InFilterClassPicked](UClass* Class)
+	{
+		InFilterClassPicked.ExecuteIfBound(Class->GetPathName());
+	});
+
+	return SNew(SBox)
+		.WidthOverride(280)
+		[
+			SNew(SVerticalBox)
+			+ SVerticalBox::Slot()
+			.AutoHeight()
+			.MaxHeight(500)
+			[
+				FModuleManager::LoadModuleChecked<FClassViewerModule>("ClassViewer").CreateClassViewer(Options, ClassPicked)
+			]
+		];	
+}
+
+/** Helper function to have user pick a (existing) package for saving a filter preset */
+static bool GetSavePresetPackageName(FString& OutPackageName)
+{
+	FAssetToolsModule& AssetToolsModule = FModuleManager::LoadModuleChecked<FAssetToolsModule>("AssetTools");
+
+	const FString DefaultName = LOCTEXT("NewFilterPreset", "NewFilterPreset").ToString();	
+	const FString DefaultPath = TEXT("/TraceSourceFilters/");
+
+	FString UniqueAssetName;
+	FString UniquePackageName;
+	AssetToolsModule.Get().CreateUniqueAssetName(DefaultPath / DefaultName, TEXT(""), UniquePackageName, UniqueAssetName);	
+
+	FSaveAssetDialogConfig SaveAssetDialogConfig;
+	{
+		SaveAssetDialogConfig.DefaultPath = DefaultPath;
+		SaveAssetDialogConfig.DefaultAssetName = FPaths::GetCleanFilename(UniqueAssetName);
+		SaveAssetDialogConfig.AssetClassNames.Add(USourceFilterCollection::StaticClass()->GetFName());
+		SaveAssetDialogConfig.ExistingAssetPolicy = ESaveAssetDialogExistingAssetPolicy::AllowButWarn;
+		SaveAssetDialogConfig.DialogTitleOverride = LOCTEXT("SaveSourceFilterPresetDialogTitle", "Save Filter Preset");
+	}
+
+	FContentBrowserModule& ContentBrowserModule = FModuleManager::LoadModuleChecked<FContentBrowserModule>("ContentBrowser");
+	FString SaveObjectPath = ContentBrowserModule.Get().CreateModalSaveAssetDialog(SaveAssetDialogConfig);
+
+	if (!SaveObjectPath.IsEmpty())
+	{
+		OutPackageName = FPackageName::ObjectPathToPackageName(SaveObjectPath);
+
+		FText OutError;
+		return FFileHelper::IsFilenameValidForSaving(OutPackageName, OutError);
+	}
+
+	return false;
+}
+
+void FEditorSessionSourceFilterService::OnSaveAsPreset()
+{
+	FString PackageName;
+	if (!GetSavePresetPackageName(PackageName))
+	{
+		return;
+	}
+
+	FScopedTransaction Transaction(*FEditorSessionSourceFilterService::TransactionContext, LOCTEXT("SaveAsPreset", "Save As Preset"), nullptr);
+
+	// Saving into a new package
+	const FString NewAssetName = FPackageName::GetLongPackageAssetName(PackageName);
+	UPackage* NewPackage = CreatePackage(nullptr, *PackageName);
+	USourceFilterCollection* NewPreset = NewObject<USourceFilterCollection>(NewPackage, *NewAssetName, RF_Public | RF_Standalone | RF_Transactional);
+
+	if (NewPreset)
+	{
+		NewPreset->CopyData(FilterCollection);
+		NewPreset->MarkPackageDirty();
+		FAssetRegistryModule::AssetCreated(NewPreset);
+
+		// Forcefully save the preset when initially created 
+		FString PackageFilename;
+		if (FPackageName::TryConvertLongPackageNameToFilename(NewPackage->GetName(), PackageFilename, FPackageName::GetAssetPackageExtension()))
+		{
+			SavePackageHelper(NewPackage, PackageFilename);
+		}
+		
+		UpdateTimeStamp();
+	}
+}
+
+TSharedPtr<FExtender> FEditorSessionSourceFilterService::GetExtender()
+{
+	if (!Extender.IsValid())
+	{
+		Extender = MakeShared<FExtender>();
+
+		const FName FilterSectionName = TEXT("FilterOptionsMenu");
+		Extender->AddMenuExtension(FilterSectionName, EExtensionHook::After, TSharedPtr<FUICommandList>(), 
+			FMenuExtensionDelegate::CreateLambda([this](FMenuBuilder& MenuBuilder) -> void
+			{
+				MenuBuilder.AddSubMenu(LOCTEXT("FilterPresetsLabel", "Filter Presets"),
+					LOCTEXT("FilterPresetsTooltip", "Saves currently set of Filters as the default setup."),
+					FNewMenuDelegate::CreateLambda([this](FMenuBuilder& SubMenuBuilder)
+						{
+							SubMenuBuilder.BeginSection(NAME_None, LOCTEXT("FiltersPresetLabel", "Filter Preset(s)"));
+
+							SubMenuBuilder.AddMenuEntry(
+								LOCTEXT("SaveFiltersPresetLabel", "Save As"),
+								LOCTEXT("SaveFiltersPresetTooltip", "Saves the currently set of Filters as a preset."),
+								FSlateIcon(FSourceFilterStyle::Get().GetStyleSetName(), "SourceFilter.SavePreset"),
+								FUIAction(
+									FExecuteAction::CreateLambda([this]()
+									{
+										OnSaveAsPreset();
+									}),
+									FCanExecuteAction::CreateLambda([this]()
+									{
+										return FilterCollection->GetFilters().Num() > 0;
+									})
+								)
+							);
+
+							SubMenuBuilder.AddMenuSeparator();
+
+							FAssetPickerConfig AssetPickerConfig;
+							{
+								AssetPickerConfig.SelectionMode = ESelectionMode::Single;
+								AssetPickerConfig.InitialAssetViewType = EAssetViewType::Column;
+								AssetPickerConfig.bFocusSearchBoxWhenOpened = true;
+								AssetPickerConfig.bAllowNullSelection = false;
+								AssetPickerConfig.bShowBottomToolbar = true;
+								AssetPickerConfig.bAutohideSearchBar = false;
+								AssetPickerConfig.bAllowDragging = false;
+								AssetPickerConfig.bCanShowClasses = false;
+								AssetPickerConfig.bShowPathInColumnView = true;
+								AssetPickerConfig.bShowTypeInColumnView = false;
+								AssetPickerConfig.bSortByPathInColumnView = false;
+
+								AssetPickerConfig.AssetShowWarningText = LOCTEXT("NoPresets_Warning", "No Presets Found");
+								AssetPickerConfig.Filter.ClassNames.Add(USourceFilterCollection::StaticClass()->GetFName());
+								AssetPickerConfig.OnAssetDoubleClicked = FOnAssetDoubleClicked::CreateLambda([this](const FAssetData& InAssetData)
+								{
+									if (UObject* Asset = InAssetData.GetAsset())
+									{
+										if (USourceFilterCollection* PresetCollection = Cast<USourceFilterCollection>(Asset))
+										{
+											FilterCollection->CopyData(PresetCollection);
+											UpdateTimeStamp();
+										}
+									}
+								});
+							}
+
+							IContentBrowserSingleton& ContentBrowser = FModuleManager::LoadModuleChecked<FContentBrowserModule>("ContentBrowser").Get();
+							SubMenuBuilder.AddWidget(SNew(SBox)
+								.MinDesiredWidth(250.f)
+								.MinDesiredHeight(400.f)
+								[
+									ContentBrowser.CreateAssetPicker(AssetPickerConfig)
+								], FText::GetEmpty()
+							);
+
+							SubMenuBuilder.EndSection();
+						}
+					)
+				);
+			})
+		);
+	}
+
+	return Extender;
+}
+
+void FEditorSessionSourceFilterService::GetWorldObjects(TArray<TSharedPtr<FWorldObject>>& OutWorldObjects)
+{
+	HashToWorld.Empty();
+
+	const TArray<const UWorld*> Worlds = FTraceWorldFiltering::GetWorlds();
+	for (const UWorld* World : Worlds)
+	{
+		FString WorldName;
+		FTraceWorldFiltering::GetWorldDisplayString(World, WorldName);
+		const uint32 WorldHash = GetTypeHash(World);
+
+		TSharedPtr<FWorldObject> WorldObject = MakeShared<FWorldObject>(WorldName, (uint8)World->WorldType, CAN_TRACE_OBJECT(World), WorldHash);
+		OutWorldObjects.Add(WorldObject);
+		HashToWorld.Add(WorldHash, World);
+	}
+}
+
+void FEditorSessionSourceFilterService::SetWorldTraceability(TSharedRef<FWorldObject> InWorldObject, bool bState)
+{
+	if (const UWorld** WorldPtr = HashToWorld.Find(InWorldObject->GetHash()))
+	{
+		FTraceWorldFiltering::SetWorldState(*WorldPtr, bState);
+	}
+}
+
+const TArray<TSharedPtr<IWorldTraceFilter>>& FEditorSessionSourceFilterService::GetWorldFilters()
+{
+	return WorldFilters;
+}
+
+void FEditorSessionSourceFilterService::PostUndo(bool bSuccess)
+{
+	UpdateTimeStamp();
+}
+
+void FEditorSessionSourceFilterService::PostRedo(bool bSuccess)
+{
+	UpdateTimeStamp();
+}
+
+void FEditorSessionSourceFilterService::MakeFilterSet(TSharedRef<const IFilterObject> ExistingFilter, TSharedRef<const IFilterObject> ExistingFilterOther)
+{
+	const FScopedTransaction Transaction(*FEditorSessionSourceFilterService::TransactionContext, LOCTEXT("CreateFilterSet", "Creating Filter Set"), FilterCollection);
+	FilterCollection->Modify();
+
+	FilterCollection->MakeFilterSet(CastChecked<UDataSourceFilter>(ExistingFilter->GetFilter()), CastChecked<UDataSourceFilter>(ExistingFilterOther->GetFilter()), EFilterSetMode::AND);
+
+	UpdateTimeStamp();
+}
+
+void FEditorSessionSourceFilterService::MakeFilterSet(TSharedRef<const IFilterObject> ExistingFilter, EFilterSetMode Mode)
+{
+	const FScopedTransaction Transaction(*FEditorSessionSourceFilterService::TransactionContext, LOCTEXT("CreateFilterSet", "Creating Filter Set"), FilterCollection);
+	FilterCollection->Modify();
+
+	FilterCollection->ConvertFilterToSet(CastChecked<UDataSourceFilter>(ExistingFilter->GetFilter()), Mode);
+
+	UpdateTimeStamp();
+}
+
+void FEditorSessionSourceFilterService::MakeTopLevelFilter(TSharedRef<const IFilterObject> Filter)
+{
+	const FScopedTransaction Transaction(*FEditorSessionSourceFilterService::TransactionContext, LOCTEXT("MakeTopLevelFilter", "Make Top Level Filter"), FilterCollection);
+	FilterCollection->Modify();
+
+	FilterCollection->MoveFilter(CastChecked<UDataSourceFilter>(Filter->GetFilter()), nullptr);
+
+	UpdateTimeStamp();
+}
+
+void FEditorSessionSourceFilterService::PopulateTreeView(FTreeViewDataBuilder& InBuilder)
+{
+	const TArray<UDataSourceFilter*>& Filters = FilterCollection->GetFilters();
+	for (UDataSourceFilter* Filter : Filters)
+	{
+		if (Filter)
+		{
+			TSharedRef<IFilterObject> FilterObject = AddFilterObjectToDataBuilder(Filter, InBuilder);
+			InBuilder.AddFilterObject(FilterObject);
+		}
+	}
+}
+
+void FEditorSessionSourceFilterService::AddReferencedObjects(FReferenceCollector& Collector)
+{
+	for (UBlueprint* Blueprint : DelegateRegisteredBlueprints)
+	{
+		Collector.AddReferencedObject(Blueprint);
+	}
+}
+
+TSharedRef<IFilterObject> FEditorSessionSourceFilterService::AddFilterObjectToDataBuilder(UDataSourceFilter* Filter, FTreeViewDataBuilder& InBuilder)
+{
+	if (UDataSourceFilterSet* FilterSet = Cast<UDataSourceFilterSet>(Filter))
+	{
+		// Deal with children as well
+		TArray<TSharedPtr<IFilterObject>> ChildObjects;
+		for (UDataSourceFilter* ChildFilter : FilterSet->GetFilters())
+		{
+			if (ChildFilter)
+			{
+				TSharedRef<IFilterObject> ChildFilterObject = AddFilterObjectToDataBuilder(ChildFilter, InBuilder);
+				ChildObjects.Add(ChildFilterObject);
+			}
+		}
+
+		TSharedRef<IFilterObject> FilterSetObject = MakeShared<FFilterSetObject>(*Cast<IDataSourceFilterSetInterface>(FilterSet), *Cast<IDataSourceFilterInterface>(Filter), ChildObjects, AsShared());
+
+		InBuilder.AddChildFilterObject(ChildObjects, FilterSetObject);
+
+		return FilterSetObject;
+	}
+	else
+	{
+		TSharedRef<IFilterObject> FilterObject = MakeShared<FFilterObject>(*Cast<IDataSourceFilterInterface>(Filter), AsShared());
+		
+		if (Filter->GetClass()->ClassGeneratedBy)
+		{
+			UBlueprint* Blueprint = CastChecked<UBlueprint>(Filter->GetClass()->ClassGeneratedBy);
+			Blueprint->OnCompiled().AddSP(this, &FEditorSessionSourceFilterService::OnBlueprintCompiled);
+			DelegateRegisteredBlueprints.Add(Blueprint);
+		}
+
+		return FilterObject;
+	}
+}
+
+void FEditorSessionSourceFilterService::SetupWorldFilters()
+{
+	WorldFilters.Add(
+		MakeShared<FWorldTypeTraceFilter>([this](uint8 TypeValue, bool bState) -> void
+		{
+			FTraceWorldFiltering::SetStateByWorldType((EWorldType::Type)TypeValue, bState);
+		},
+		[this](uint8 TypeValue) -> bool
+		{
+			return FTraceWorldFiltering::IsWorldTypeTraceable((EWorldType::Type)TypeValue);
+		})
+	);
+
+	WorldFilters.Add(
+		MakeShared<FWorldNetModeTraceFilter>([this](uint8 TypeValue, bool bState) -> void
+		{
+			FTraceWorldFiltering::SetStateByWorldNetMode((ENetMode)TypeValue, bState);
+		},
+		[this](uint8 TypeValue) -> bool
+		{
+			return FTraceWorldFiltering::IsWorldNetModeTraceable((ENetMode)TypeValue);
+		})
+	);
+}
+
+void FEditorSessionSourceFilterService::OnBlueprintCompiled(UBlueprint* InBlueprint)
+{
+	if (InBlueprint)
+	{
+		UpdateTimeStamp();
+
+		for (UBlueprint* Blueprint : DelegateRegisteredBlueprints)
+		{
+			Blueprint->OnCompiled().RemoveAll(this);
+		}
+		DelegateRegisteredBlueprints.Empty();
+	}
+}
+
+#endif // WITH_ENGINE
+
+#undef LOCTEXT_NAMESPACE // "FEditorSourceFilterService"
diff --git a/Engine/Plugins/Developer/TraceSourceFiltering/Source/SourceFilteringEditor/Private/EditorSessionSourceFilterService.h b/Engine/Plugins/Developer/TraceSourceFiltering/Source/SourceFilteringEditor/Private/EditorSessionSourceFilterService.h
new file mode 100644
--- /dev/null
+++ b/Engine/Plugins/Developer/TraceSourceFiltering/Source/SourceFilteringEditor/Private/EditorSessionSourceFilterService.h
@@ -0,0 +1,101 @@
+// Copyright Epic Games, Inc. All Rights Reserved.
+
+#pragma once
+
+#if WITH_EDITOR
+
+#include "ISessionSourceFilterService.h"
+
+#include "UObject/GCObject.h"
+#include "EditorUndoClient.h"
+
+#include "SourceFilterCollection.h"
+#include "TreeViewBuilder.h"
+#include "DataSourceFilter.h"
+#include "TraceSourceFilteringSettings.h"
+
+namespace Trace
+{
+	class IAnalysisSession;
+	typedef uint64 FSessionHandle;
+}
+
+class UDataSourceFilter;
+class UDataSourceFilterSet;
+
+/** Editor implementation of ISessionSourceFilterService, interfaces directly with Engine level filtering systems and settings */
+class FEditorSessionSourceFilterService : public ISessionSourceFilterService, public FGCObject, public FEditorUndoClient
+{
+public:
+	FEditorSessionSourceFilterService();
+	virtual ~FEditorSessionSourceFilterService();
+
+	/** Begin ISessionSourceFilterService overrides */
+	virtual void PopulateTreeView(FTreeViewDataBuilder& InBuilder) override;	
+	virtual void AddFilter(const FString& FilterClassName) override;
+	virtual void AddFilterToSet(TSharedRef<const IFilterObject> FilterSet, const FString& FilterClassName) override;
+	virtual void AddFilterToSet(TSharedRef<const IFilterObject> FilterSet, TSharedRef<const IFilterObject> ExistingFilter) override;
+	virtual void MakeFilterSet(TSharedRef<const IFilterObject> ExistingFilter, EFilterSetMode Mode) override;
+	virtual void MakeFilterSet(TSharedRef<const IFilterObject> ExistingFilter, TSharedRef<const IFilterObject> ExistingFilterOther) override;
+	virtual void MakeTopLevelFilter(TSharedRef<const IFilterObject> Filter) override;
+	virtual void RemoveFilter(TSharedRef<const IFilterObject> InFilter)  override;
+	virtual void SetFilterSetMode(TSharedRef<const IFilterObject> InFilter, EFilterSetMode Mode) override;
+	virtual void SetFilterState(TSharedRef<const IFilterObject> InFilter, bool bState) override;
+	virtual void ResetFilters() override;	
+	virtual const FDateTime& GetTimestamp() const override { return Timestamp; }	
+	virtual void UpdateFilterSettings(UTraceSourceFilteringSettings* InSettings) override {}
+	virtual UTraceSourceFilteringSettings* GetFilterSettings() override;
+	virtual bool IsActionPending() const override;
+	virtual TSharedRef<SWidget> GetFilterPickerWidget(FOnFilterClassPicked InFilterClassPicked) override;
+	virtual TSharedPtr<FExtender> GetExtender() override;
+	virtual void GetWorldObjects(TArray<TSharedPtr<FWorldObject>>& OutWorldObjects) override;
+	virtual void SetWorldTraceability(TSharedRef<FWorldObject> InWorldObject, bool bState) override;
+	virtual const TArray<TSharedPtr<IWorldTraceFilter>>& GetWorldFilters() override;
+	/** End ISessionSourceFilterService overrides */	
+
+	/** Begin FEditorUndoClient overrides*/
+	virtual void PostUndo(bool bSuccess) override;
+	virtual void PostRedo(bool bSuccess) override;
+	/** End FEditorUndoClient overrides */
+
+	/** Begin FGCObject overrides*/
+	virtual void AddReferencedObjects(FReferenceCollector& Collector) override;
+	/** End FGCObject overrides */
+
+protected:
+	/** Updates the timestamp provided through GetTimestamp, indicating that the contained data has changed  */
+	void UpdateTimeStamp();
+
+	/** Callback for whenever the user opts to save the current filtering state as a preset */
+	void OnSaveAsPreset();
+
+	/** Callback for whenever any Filter instance blueprint class gets (re)compiled */
+	void OnBlueprintCompiled(UBlueprint* InBlueprint);
+
+	/** Helper function to populate FTreeViewDataBuilder with a specific FilterObject */
+	TSharedRef<IFilterObject> AddFilterObjectToDataBuilder(UDataSourceFilter* Filter, FTreeViewDataBuilder& InBuilder);
+
+	void SetupWorldFilters();
+protected:
+	/** Blueprints on which a OnCompiled calllback has been registered, used to refresh data and repopulate the UI */
+	TArray<UBlueprint*> DelegateRegisteredBlueprints;
+
+	/** Filter collection this session represents, retrieved from FTraceSourceFiltering */
+	USourceFilterCollection* FilterCollection;
+
+	/** Timestamp used to rely data updates to anyone polling against GetTimeStamp */
+	FDateTime Timestamp;
+	
+	/** UI extender, used to insert preset functionality into STraceSourceFilteringWindow */
+	TSharedPtr<FExtender> Extender;
+
+	/** World filtering data */
+	TMap<uint32, const UWorld*> HashToWorld;
+	TArray<TSharedPtr<IWorldTraceFilter>> WorldFilters;
+
+	/** Transaction context naming for all scoped transactions performed in thsi class */
+	static FString TransactionContext;
+
+};
+
+#endif // WITH_ENGINE
\ No newline at end of file
diff --git a/Engine/Plugins/Developer/TraceSourceFiltering/Source/SourceFilteringEditor/Private/FilterDragDropOperation.h b/Engine/Plugins/Developer/TraceSourceFiltering/Source/SourceFilteringEditor/Private/FilterDragDropOperation.h
new file mode 100644
--- /dev/null
+++ b/Engine/Plugins/Developer/TraceSourceFiltering/Source/SourceFilteringEditor/Private/FilterDragDropOperation.h
@@ -0,0 +1,91 @@
+// Copyright Epic Games, Inc. All Rights Reserved.
+#pragma once
+
+#include "CoreMinimal.h"
+#include "Input/DragAndDrop.h"
+#include "Widgets/DeclarativeSyntaxSupport.h"
+#include "EditorStyleSet.h"
+#include "Widgets/Layout/SBorder.h"
+#include "Widgets/Text/STextBlock.h"
+#include "Widgets/SBoxPanel.h"
+#include "Widgets/Images/SImage.h"
+
+#include "SourceFilterStyle.h"
+
+class IFilterObject;
+
+class FFilterDragDropOp : public FDragDropOperation
+{
+public:
+	DRAG_DROP_OPERATOR_TYPE(FFilterDragDropOp, FDragDropOperation)
+
+	TWeakPtr<class IFilterObject> FilterObject;
+
+	/** The widget decorator to use */
+	virtual TSharedPtr<SWidget> GetDefaultDecorator() const override
+	{
+		return SNew(SBorder)
+			.BorderImage(FSourceFilterStyle::GetBrush("SourceFilter.DragDrop.Border"))
+			.Content()
+			[
+				SNew(SHorizontalBox)
+				+SHorizontalBox::Slot()
+				.AutoWidth()
+				.Padding(2, 1, 0, 0)
+				.VAlign(VAlign_Center)
+				[
+					SNew(STextBlock)
+					.TextStyle(FSourceFilterStyle::Get(), "SourceFilter.TextStyle")
+					.Font(FEditorStyle::Get().GetFontStyle("FontAwesome.12"))
+					.Text(this, &FFilterDragDropOp::GetIconText)
+				]
+				+SHorizontalBox::Slot()
+				.AutoWidth()
+				.Padding(2, 0, 0, 0)
+				.VAlign(VAlign_Center)
+				[
+					SNew(STextBlock)
+					.Text(this, &FFilterDragDropOp::GetHoverText)
+				]
+			];
+	}
+
+	FText GetHoverText() const
+	{
+		return Text;	
+	}
+
+	FText GetIconText() const
+	{
+		return IconText;
+	}
+
+	void SetIconText(FText InText)
+	{
+		IconText = InText;
+	}
+
+	void SetText(FText InText)
+	{
+		Text = InText;
+	}
+
+	void Reset()
+	{
+		Text = FText::FromString(TEXT("Invalid Operation"));
+		IconText = FText::FromString(FString(TEXT("\xf05e")));
+	}
+	
+	static TSharedRef<FFilterDragDropOp> New(TSharedRef<class IFilterObject> InFilterObject)
+	{
+		TSharedRef<FFilterDragDropOp> Operation = MakeShareable(new FFilterDragDropOp);
+		Operation->FilterObject = InFilterObject;
+		Operation->Reset();
+		Operation->Construct();
+		return Operation;
+	}
+
+private:
+	FText IconText;
+	FText Text;
+};
diff --git a/Engine/Plugins/Developer/TraceSourceFiltering/Source/SourceFilteringEditor/Private/FilterObject.cpp b/Engine/Plugins/Developer/TraceSourceFiltering/Source/SourceFilteringEditor/Private/FilterObject.cpp
new file mode 100644
--- /dev/null
+++ b/Engine/Plugins/Developer/TraceSourceFiltering/Source/SourceFilteringEditor/Private/FilterObject.cpp
@@ -0,0 +1,145 @@
+// Copyright Epic Games, Inc. All Rights Reserved.
+
+#include "FilterObject.h"
+
+#include "Widgets/Layout/SWrapBox.h"
+
+#include "ISessionSourceFilterService.h"
+#include "FilterDragDropOperation.h"
+#include "IDataSourceFilterInterface.h"
+#include "SourceFilterStyle.h"
+
+#if WITH_EDITOR
+#include "Editor/EditorEngine.h"
+#include "Editor.h"
+#include "Subsystems/AssetEditorSubsystem.h"
+#endif // WITH_EDITOR
+
+#define LOCTEXT_NAMESPACE "FilterObject"
+
+FFilterObject::FFilterObject(IDataSourceFilterInterface& InFilter, TSharedRef<ISessionSourceFilterService> InSessionFilterService) : SessionFilterService(InSessionFilterService)
+{
+	WeakInterface = InFilter;
+}
+
+FText FFilterObject::GetDisplayText() const
+{
+	FText Text;
+	if (GetFilter())
+	{
+		WeakInterface->Execute_GetDisplayText(WeakInterface.GetObject(), Text);
+	}
+	return Text;
+}
+
+FText FFilterObject::GetToolTipText() const
+{
+	FText Text;
+	if (GetFilter())
+	{
+		WeakInterface->Execute_GetToolTipText(WeakInterface.GetObject(), Text);
+	}
+	return Text;
+}
+
+UObject* FFilterObject::GetFilter() const
+{
+	return WeakInterface.GetObject();
+}
+
+bool FFilterObject::IsFilterEnabled() const
+{
+	if (GetFilter())
+	{
+		return WeakInterface->IsEnabled();
+	}
+
+	return false;
+}
+
+TSharedRef<SWidget> FFilterObject::MakeWidget(TSharedRef<SWrapBox> ParentWrapBox) const
+{
+	TSharedPtr<SHorizontalBox> WidgetBox;
+
+	ParentWrapBox->AddSlot()
+	[
+		SAssignNew(WidgetBox, SHorizontalBox)
+		+ SHorizontalBox::Slot()
+		.AutoWidth()
+		.Padding(2, 0, 0, 0)
+		[
+			SNew(SBorder)
+			.BorderImage(FSourceFilterStyle::GetBrush("SourceFilter.FilterBrush"))
+			.ForegroundColor(FLinearColor::White)
+			.Padding(FMargin(4.f))
+			[
+				SNew(STextBlock)
+				.TextStyle(FSourceFilterStyle::Get(), "SourceFilter.TextStyle")
+				.Text_Lambda([this]() { return GetDisplayText(); })
+				.ToolTipText_Lambda([this]() { return GetToolTipText(); })
+			]
+		]
+	];
+
+	TAttribute<bool> FilterEnabled = TAttribute<bool>::Create(TAttribute<bool>::FGetter::CreateSP(this, &FFilterObject::IsFilterEnabled));
+	WidgetBox->SetEnabled(FilterEnabled);
+
+	return ParentWrapBox;
+}
+
+FReply FFilterObject::HandleDragEnter(const FDragDropEvent& DragDropEvent)
+{
+	TSharedPtr<FFilterDragDropOp> FilterDragOp = DragDropEvent.GetOperationAs<FFilterDragDropOp>();
+
+	if (FilterDragOp.IsValid())
+	{
+		// Check if the dragged object is actually a Filter 
+		if (FilterDragOp->FilterObject.Pin()->GetFilter()->Implements<UDataSourceFilterInterface>())
+		{
+			// Ensure we are not dragging onto ourselves, if not allow for Creating a Filter Set containing both the dragged and this Filter
+			if (FilterDragOp->FilterObject.Pin() != AsShared())
+			{
+				FilterDragOp->SetIconText(FText::FromString(TEXT("\xf0fe")));
+				FilterDragOp->SetText(LOCTEXT("CreateFilterToSetDragOpLabel", "Create Filter Set"));
+			}
+			else
+			{
+				FilterDragOp->Reset();
+			}
+
+			return FReply::Handled();
+		}
+	}
+
+	return FReply::Unhandled();
+}
+
+void FFilterObject::HandleDragLeave(const FDragDropEvent& DragDropEvent)
+{
+	TSharedPtr<FFilterDragDropOp> FilterDragOp = DragDropEvent.GetOperationAs<FFilterDragDropOp>();
+	if (FilterDragOp.IsValid())
+	{
+		FilterDragOp->Reset();
+	}
+}
+
+FReply FFilterObject::HandleDrop(const FDragDropEvent& DragDropEvent)
+{
+	TSharedPtr<FFilterDragDropOp> FilterDragOp = DragDropEvent.GetOperationAs<FFilterDragDropOp>();
+
+	if (FilterDragOp.IsValid())
+	{
+		if (FilterDragOp->FilterObject.Pin()->GetFilter()->Implements<UDataSourceFilterInterface>() && FilterDragOp->FilterObject.Pin() != AsShared())
+		{
+			// Request to create a filter set contain both the dragged filter and this
+			SessionFilterService->MakeFilterSet(AsShared(), FilterDragOp->FilterObject.Pin()->AsShared());
+			return FReply::Handled();
+		}
+
+		return FReply::Handled();
+	}
+
+	return FReply::Unhandled();
+}
+
+#undef LOCTEXT_NAMESPACE // "FilterObject"
\ No newline at end of file
diff --git a/Engine/Plugins/Developer/TraceSourceFiltering/Source/SourceFilteringEditor/Private/FilterObject.h b/Engine/Plugins/Developer/TraceSourceFiltering/Source/SourceFilteringEditor/Private/FilterObject.h
new file mode 100644
--- /dev/null
+++ b/Engine/Plugins/Developer/TraceSourceFiltering/Source/SourceFilteringEditor/Private/FilterObject.h
@@ -0,0 +1,31 @@
+// Copyright Epic Games, Inc. All Rights Reserved.
+
+#pragma once
+
+#include "IFilterObject.h"
+#include "UObject/WeakInterfacePtr.h"
+
+class ISessionSourceFilterService;
+class IDataSourceFilterInterface;
+
+class FFilterObject : public IFilterObject
+{
+public:
+	FFilterObject(IDataSourceFilterInterface& InFilter, TSharedRef<ISessionSourceFilterService> InSessionFilterService);
+
+	/** Begin IFilterObject overrides */
+	virtual FText GetDisplayText() const;
+	virtual FText GetToolTipText() const;
+	virtual UObject* GetFilter() const;
+	virtual bool IsFilterEnabled() const;	
+	virtual TSharedRef<SWidget> MakeWidget(TSharedRef<SWrapBox> ParentWrapBox) const;
+	virtual FReply HandleDragEnter(const FDragDropEvent& DragDropEvent) override;
+	virtual void HandleDragLeave(const FDragDropEvent& DragDropEvent) override;
+	virtual FReply HandleDrop(const FDragDropEvent& DragDropEvent) override;
+	/** End IFilterObject overrides */
+
+protected:
+	/** Weak interface pointer retrieved from Data Source Filter instance this object represents */
+	TWeakInterfacePtr<IDataSourceFilterInterface> WeakInterface;
+	TSharedPtr<ISessionSourceFilterService> SessionFilterService;
+};
diff --git a/Engine/Plugins/Developer/TraceSourceFiltering/Source/SourceFilteringEditor/Private/FilterSetObject.cpp b/Engine/Plugins/Developer/TraceSourceFiltering/Source/SourceFilteringEditor/Private/FilterSetObject.cpp
new file mode 100644
--- /dev/null
+++ b/Engine/Plugins/Developer/TraceSourceFiltering/Source/SourceFilteringEditor/Private/FilterSetObject.cpp
@@ -0,0 +1,248 @@
+// Copyright Epic Games, Inc. All Rights Reserved.
+
+#include "FilterSetObject.h"
+
+#include "Widgets/SBoxPanel.h"
+#include "Widgets/Layout/SWrapBox.h"
+#include "Widgets/Input/SComboButton.h"
+#include "Widgets/Layout/SBorder.h"
+#include "Framework/Application/SlateApplication.h"
+#include "UObject/UObjectGlobals.h"
+#include "UObject/ObjectMacros.h"
+#include "UObject/Class.h"
+#include "UObject/Package.h"
+
+#include "ISessionSourceFilterService.h"
+#include "IDataSourceFilterInterface.h"
+#include "IDataSourceFilterSetInterface.h"
+#include "SourceFilterStyle.h"
+#include "FilterDragDropOperation.h"
+
+#define LOCTEXT_NAMESPACE "FilterSetObject"
+
+FFilterSetObject::FFilterSetObject(IDataSourceFilterSetInterface& InFilterSet, IDataSourceFilterInterface& InFilter, const TArray<TSharedPtr<IFilterObject>>& InChildFilters, TSharedRef<ISessionSourceFilterService> InSessionFilterService) : ChildFilters(InChildFilters), SessionFilterService(InSessionFilterService), FilterSetModeEnumPtr(FindObject<UEnum>(ANY_PACKAGE, TEXT("EFilterSetMode"), true))
+{
+	WeakFilter = InFilter;
+	WeakFilterSet = InFilterSet;
+}
+
+FText FFilterSetObject::GetDisplayText() const
+{
+	FText Text;
+	if (GetFilter())
+	{
+		WeakFilter->GetDisplayText_Implementation(Text);
+	}
+	return Text;
+}
+
+FText FFilterSetObject::GetToolTipText() const
+{
+	FText Text;
+	if (GetFilter())
+	{
+		WeakFilter->GetToolTipText_Implementation(Text);
+	}
+	return Text;
+}
+
+UObject* FFilterSetObject::GetFilter() const
+{
+	return WeakFilter.GetObject();
+}
+
+bool FFilterSetObject::IsFilterEnabled() const
+{
+	if (GetFilter())
+	{
+		return WeakFilter->IsEnabled();
+	}
+
+	return false;
+}
+
+void FFilterSetObject::AddFilterSetModeWidget(TSharedRef<SWrapBox> ParentWrapBox) const
+{
+	ParentWrapBox->AddSlot()
+	[
+		SNew(SHorizontalBox)
+		+ SHorizontalBox::Slot()
+		.AutoWidth()
+		.Padding(2, 0, 0, 0)
+		.VAlign(VAlign_Center)
+		[
+			SNew(SComboButton)
+			.ComboButtonStyle(FSourceFilterStyle::Get(), "GenericFilters.ComboButtonStyle")
+			.ButtonStyle(FSourceFilterStyle::Get(), FName(*("SourceFilter.FilterSetOperation." + FilterSetModeEnumPtr->GetNameStringByValue((int64)WeakFilterSet->GetFilterSetMode()))))
+			.ForegroundColor(FLinearColor::White)
+			.ContentPadding(2.0f)
+			.HasDownArrow(false)
+			.VAlign(VAlign_Center)
+			.ToolTipText_Lambda([this]() -> FText
+			{
+				return FText::Format(LOCTEXT("FilterSetModeTooltip", "{0} operator is used for this Filter Set."), FilterSetModeEnumPtr->GetDisplayNameTextByValue((int64)WeakFilterSet->GetFilterSetMode()));
+			})
+			.OnGetMenuContent(FOnGetContent::CreateLambda([this]()
+			{
+				FMenuBuilder MenuBuilder(true, TSharedPtr<FUICommandList>());
+
+				const FText LabelTextFormat = LOCTEXT("ChangeFilterSetModeLabel", "{0}");
+				const FText ToolTipTextFormat = LOCTEXT("ChangeFilterSetModeTooltip", "Set Filtering Mode to {0} operator.");
+
+				MenuBuilder.BeginSection(NAME_None, LOCTEXT("FilterModeSectionLabel", "Filter Modes"));
+
+				// Add an entry for each Filter Set Mode, except for the one currently set for this Filter Set
+				for (EFilterSetMode Mode : TEnumRange<EFilterSetMode>())
+				{
+					if (Mode != WeakFilterSet->GetFilterSetMode())
+					{
+						const FText ModeText = FilterSetModeEnumPtr->GetDisplayNameTextByValue((int64)Mode);
+
+						MenuBuilder.AddMenuEntry(
+							FText::Format(LabelTextFormat, ModeText),
+							FText::Format(ToolTipTextFormat, ModeText),
+							FSlateIcon(),
+							FUIAction(
+								FExecuteAction::CreateLambda([this, Mode]() 
+								{
+									// Request a change of this Filter Set's mode through the session service
+									SessionFilterService->SetFilterSetMode(AsShared(), Mode); 
+								})
+							)
+						);
+					}
+				}
+
+				MenuBuilder.EndSection();
+
+				return MenuBuilder.MakeWidget();
+			}))
+			.ButtonContent()
+			[
+				SNew(STextBlock)
+				.TextStyle(FSourceFilterStyle::Get(), "SourceFilter.TextStyle")
+				.Text_Lambda([this]()
+				{
+					return FilterSetModeEnumPtr->GetDisplayNameTextByValue((int64)WeakFilterSet->GetFilterSetMode());
+				})
+			]
+		]
+	];
+}
+
+TSharedRef<SWidget> FFilterSetObject::MakeWidget(TSharedRef<SWrapBox> ParentWrapBox) const
+{
+	auto AddEncapsulatingWidget = [ParentWrapBox](FText InText)
+	{
+		ParentWrapBox->AddSlot()
+		[
+			SNew(SHorizontalBox)
+			+ SHorizontalBox::Slot()
+			.AutoWidth()
+			.Padding(2, 0, 0, 0)
+			[
+				SNew(SBorder)
+				.BorderImage(FSourceFilterStyle::GetBrush("SourceFilter.FilterSetBrush"))
+				.ForegroundColor(FLinearColor::White)
+				.Padding(FMargin(4.f))
+				[
+					SNew(STextBlock)
+					.TextStyle(FSourceFilterStyle::Get(), "SourceFilter.TextStyle")
+					.Text(InText)
+				]
+			]
+		];
+	};
+
+	AddEncapsulatingWidget(LOCTEXT("OpenBracket", "("));	
+
+	// In case this is a NOT filter, only add the operation widget at the start of the filter set as such: (!FilterA, FilterB)
+	const bool bIsNOTFilterSet = WeakFilterSet->GetFilterSetMode() == EFilterSetMode::NOT;
+	if (bIsNOTFilterSet)
+	{
+		AddFilterSetModeWidget(ParentWrapBox);
+	}
+
+	for (const TSharedPtr<IFilterObject>& ChildFilter : ChildFilters)
+	{
+		ChildFilter->MakeWidget(ParentWrapBox);
+
+		// In case this is not a NOT filter, add the operation widget in between each filter (or at the end of a filter set containing a single filter) as such: (FilterA AND FilterB)
+		const bool bIsLastFilter = ChildFilter != ChildFilters.Last();
+		const bool bShouldAddFilterSetWidget = bIsLastFilter || (ChildFilters.Num() == 1 && !bIsNOTFilterSet);
+		if (bShouldAddFilterSetWidget)
+		{
+			AddFilterSetModeWidget(ParentWrapBox);
+		}
+	}
+
+	AddEncapsulatingWidget(LOCTEXT("CloseBracket", ")"));
+
+	return ParentWrapBox;
+}
+
+FReply FFilterSetObject::HandleDragEnter(const FDragDropEvent& DragDropEvent)
+{
+	TSharedPtr<FFilterDragDropOp> FilterDragOp = DragDropEvent.GetOperationAs<FFilterDragDropOp>();
+
+	if (FilterDragOp.IsValid())
+	{
+		// Ensure that the dragged object is indeed a filter
+		if (FilterDragOp->FilterObject.Pin()->GetFilter()->Implements<UDataSourceFilterInterface>())
+		{
+			// And that we're not dragging onto ourselves
+			if (FilterDragOp->FilterObject.Pin() != AsShared())
+			{
+				// In case this filter set not yet contains the dragged object, allow for adding it				
+				if (!ChildFilters.Contains(FilterDragOp->FilterObject))
+				{
+					FilterDragOp->SetIconText(FText::FromString(TEXT("\xf0fe")));
+					FilterDragOp->SetText(LOCTEXT("AddFilterToSetDragOpLabel", "Add Filter to Filter Set"));
+				}
+				else
+				{
+					FilterDragOp->Reset();
+					FilterDragOp->SetText(LOCTEXT("AlreadyChildFilterToSetDragOpLabel", "Filter is already part of Filter Set"));
+				}
+			}
+			else
+			{
+				FilterDragOp->Reset();
+			}
+
+			return FReply::Handled();
+		}
+	}
+
+	return FReply::Unhandled();
+}
+
+void FFilterSetObject::HandleDragLeave(const FDragDropEvent& DragDropEvent)
+{
+	TSharedPtr<FFilterDragDropOp> FilterDragOp = DragDropEvent.GetOperationAs<FFilterDragDropOp>();
+	if (FilterDragOp.IsValid())
+	{
+		FilterDragOp->Reset();
+	}
+}
+
+FReply FFilterSetObject::HandleDrop(const FDragDropEvent& DragDropEvent)
+{
+	TSharedPtr<FFilterDragDropOp> FilterDragOp = DragDropEvent.GetOperationAs<FFilterDragDropOp>();
+
+	if (FilterDragOp.IsValid())
+	{
+		// If the dragged object is a Filter, and is not yet contained by this Filter Set, request to add it
+		if (FilterDragOp->FilterObject.Pin()->GetFilter()->Implements<UDataSourceFilterInterface>() && FilterDragOp->FilterObject.Pin() != AsShared() && !ChildFilters.Contains(FilterDragOp->FilterObject))
+		{
+			SessionFilterService->AddFilterToSet(AsShared(), FilterDragOp->FilterObject.Pin()->AsShared());
+			return FReply::Handled();
+		}
+
+		return FReply::Handled();
+	}
+
+	return FReply::Unhandled();
+}
+
+#undef LOCTEXT_NAMESPACE // "FilterSetObject"
\ No newline at end of file
diff --git a/Engine/Plugins/Developer/TraceSourceFiltering/Source/SourceFilteringEditor/Private/FilterSetObject.h b/Engine/Plugins/Developer/TraceSourceFiltering/Source/SourceFilteringEditor/Private/FilterSetObject.h
new file mode 100644
--- /dev/null
+++ b/Engine/Plugins/Developer/TraceSourceFiltering/Source/SourceFilteringEditor/Private/FilterSetObject.h
@@ -0,0 +1,46 @@
+// Copyright Epic Games, Inc. All Rights Reserved.
+
+#pragma once
+
+#include "IFilterObject.h"
+#include "UObject/WeakInterfacePtr.h"
+
+class ISessionSourceFilterService;
+class IDataSourceFilterSetInterface;
+class IDataSourceFilterInterface;
+
+class FFilterSetObject : public IFilterObject
+{
+public:
+	FFilterSetObject(IDataSourceFilterSetInterface& InFilterSet, IDataSourceFilterInterface& InFilter, const TArray<TSharedPtr<IFilterObject>>& InChildFilters, TSharedRef<ISessionSourceFilterService> InSessionFilterService);
+
+	/** Begin IFilterObject overrides */
+	virtual FText GetDisplayText() const;
+	virtual FText GetToolTipText() const;
+	virtual UObject* GetFilter() const;	
+	virtual bool IsFilterEnabled() const;
+	virtual TSharedRef<SWidget> MakeWidget(TSharedRef<SWrapBox> ParentWrapBox) const;
+	virtual FReply HandleDragEnter(const FDragDropEvent& DragDropEvent) override;
+	virtual void HandleDragLeave(const FDragDropEvent& DragDropEvent) override;
+	virtual FReply HandleDrop(const FDragDropEvent& DragDropEvent) override;
+	/** End IFilterObject overrides */
+
+protected:
+	/** Adds appropriate Filter Set operation widget to the provided SWrapBox */
+	void AddFilterSetModeWidget(TSharedRef<SWrapBox> ParentWrapBox) const;
+
+protected:
+	/** Weak filter interface pointer retrieved from Data Source Filter instance this object represents */
+	TWeakInterfacePtr<IDataSourceFilterInterface> WeakFilter;
+	/** Weak filter set interface pointer retrieved from Data Source Filter instance this object represents */
+	TWeakInterfacePtr<IDataSourceFilterSetInterface> WeakFilterSet;
+
+	/** Any filter objects which are part of this set*/
+	TArray<TSharedPtr<IFilterObject>> ChildFilters;
+
+	/** Session filter service which is responsible for generating this object */
+	TSharedPtr<ISessionSourceFilterService> SessionFilterService;
+
+	/** Cached pointer to UEnum representing EFilterSetMode */
+	const UEnum* FilterSetModeEnumPtr;
+};
diff --git a/Engine/Plugins/Developer/TraceSourceFiltering/Source/SourceFilteringEditor/Private/IFilterObject.h b/Engine/Plugins/Developer/TraceSourceFiltering/Source/SourceFilteringEditor/Private/IFilterObject.h
new file mode 100644
--- /dev/null
+++ b/Engine/Plugins/Developer/TraceSourceFiltering/Source/SourceFilteringEditor/Private/IFilterObject.h
@@ -0,0 +1,43 @@
+// Copyright Epic Games, Inc. All Rights Reserved.
+
+#pragma once
+
+#include "Templates/SharedPointer.h"
+#include "Input/DragAndDrop.h"
+#include "Input/Reply.h"
+
+class FText;
+class UObject;
+class SWidget;
+class SWrapBox;
+
+/** Interface for any FilterObject representation throughout the Filtering UI */
+class IFilterObject : public TSharedFromThis<IFilterObject>
+{
+public:
+	virtual ~IFilterObject() {}
+
+	/** Return the display text to be displayed for this object */
+	virtual FText GetDisplayText() const = 0;
+	
+	/** Return the tooltip to be displayed for this object */
+	virtual FText GetToolTipText() const = 0;	
+
+	/** Return Filter UObject */
+	virtual UObject* GetFilter() const = 0;
+
+	/** Return whether or not this filter object is currently enabled*/
+	virtual bool IsFilterEnabled() const = 0;
+
+	/** Generate widget representing this object */
+	virtual TSharedRef<SWidget> MakeWidget(TSharedRef<SWrapBox> ParentWrapBox) const = 0;
+
+	/** Handle a drag and drop enter event */
+	virtual FReply HandleDragEnter(const FDragDropEvent& DragDropEvent) = 0;
+
+	/** Handle a drag and drop leave event */
+	virtual void HandleDragLeave(const FDragDropEvent& DragDropEvent) = 0;
+
+	/** Handle a drag and drop drop event */
+	virtual FReply HandleDrop(const FDragDropEvent& DragDropEvent) = 0;
+};
diff --git a/Engine/Plugins/Developer/TraceSourceFiltering/Source/SourceFilteringEditor/Private/ISessionSourceFilterService.h b/Engine/Plugins/Developer/TraceSourceFiltering/Source/SourceFilteringEditor/Private/ISessionSourceFilterService.h
new file mode 100644
--- /dev/null
+++ b/Engine/Plugins/Developer/TraceSourceFiltering/Source/SourceFilteringEditor/Private/ISessionSourceFilterService.h
@@ -0,0 +1,97 @@
+// Copyright Epic Games, Inc. All Rights Reserved.
+
+#pragma once
+
+#include "Containers/Array.h"
+#include "Containers/UnrealString.h"
+#include "Delegates/DelegateCombinations.h"
+#include "Templates/SharedPointer.h"
+#include "Framework/MultiBox/MultiBoxBuilder.h"
+
+enum class EFilterSetMode : uint8;
+
+class IFilterObject;
+struct FTreeViewDataBuilder;
+class SWidget;
+
+class FSourceFilterCollection;
+class UTraceSourceFilteringSettings;
+class FWorldObject;
+
+DECLARE_DELEGATE_OneParam(FOnFilterClassPicked, FString /*ClassName*/);
+
+/** Interface for implementing a World filter, setting whether or not a specific UWorld and its contained objects can Trace out data (events) */
+class IWorldTraceFilter : public TSharedFromThis<IWorldTraceFilter>
+{
+public:
+	virtual ~IWorldTraceFilter() {}
+	virtual FText GetDisplayText() = 0;
+	virtual FText GetToolTipText() = 0;
+	virtual void PopulateMenuBuilder(FMenuBuilder& InBuilder) = 0;
+};
+
+/** Interface providing interaction with an running instance's its Trace and World Filtering systems */
+class ISessionSourceFilterService : public TSharedFromThis<ISessionSourceFilterService>
+{
+public:
+	virtual ~ISessionSourceFilterService() {}
+
+	/** Returns timestamp representing last point of update */
+	virtual const FDateTime& GetTimestamp() const = 0;
+
+	/** Returns whether or not a request action is pending complete */
+	virtual bool IsActionPending() const = 0;
+
+	/** Adds a Filter instance of the provided UClass name */
+	virtual void AddFilter(const FString& FilterClassName) = 0;
+
+	/** Removes a specific Filter (Set) instance */
+	virtual void RemoveFilter(TSharedRef<const IFilterObject> InFilter) = 0;
+	
+	/** Add a Filter instance, of the provided UClass name, to the specified Filter Set instance */
+	virtual void AddFilterToSet(TSharedRef<const IFilterObject> FilterSet, const FString& FilterClassName) = 0;	
+
+	/** Add a specific Filter instance to the specified Filter Set instance */
+	virtual void AddFilterToSet(TSharedRef<const IFilterObject> FilterSet, TSharedRef<const IFilterObject> ExistingFilter) = 0;
+
+	/** Moves the specific filter to root level in the filtering tree */
+	virtual void MakeTopLevelFilter(TSharedRef<const IFilterObject> Filter) = 0;
+	
+	/** Creates a new Filter Set instance, with the provided filter mode, replacing and containing the specific Filter instance*/
+	virtual void MakeFilterSet(TSharedRef<const IFilterObject> ExistingFilter, EFilterSetMode Mode) = 0;
+
+	/** Creates a new Filter Set instance, with default AND filter mode, replacing and containing the specific Filter instances*/
+	virtual void MakeFilterSet(TSharedRef<const IFilterObject> ExistingFilter, TSharedRef<const IFilterObject> ExistingFilterOther) = 0;
+
+	/** Sets the state of a specific Filter instance, true= enabled; false= disabled */
+	virtual void SetFilterState(TSharedRef<const IFilterObject> InFilter, bool bState) = 0;
+
+	/** Sets the Filtering Mode for a specific Filter Set instance */
+	virtual void SetFilterSetMode(TSharedRef<const IFilterObject> InFilter, EFilterSetMode Mode) = 0;
+
+	/** Resets the complete filtering tree, removing all Filter instances */
+	virtual void ResetFilters() = 0;
+
+	/** Updates the Filtering Settings for this specific session */
+  	virtual void UpdateFilterSettings(UTraceSourceFilteringSettings* InSettings) = 0;
+	
+	/** Retrieves the current state of the Filtering Settings for this specific session */
+	virtual UTraceSourceFilteringSettings* GetFilterSettings() = 0;
+	
+	/** Request to populate a TreeView using the filter (set) hierarchy */
+	virtual void PopulateTreeView(FTreeViewDataBuilder& InBuilder) = 0;
+	
+	/** Returns a slate widget used for picking a Filter class, which will execute the passed in delegate on selection/completion */
+	virtual TSharedRef<SWidget> GetFilterPickerWidget(FOnFilterClassPicked InFilterClassPicked) = 0;
+
+	/** Returns an FExtender instance which is incorporated anytime a context menu (MenuBuilder) is created */
+	virtual TSharedPtr<FExtender> GetExtender() = 0;
+
+	/** Returns an FWorldObject for each active UWorld instance */
+	virtual void GetWorldObjects(TArray<TSharedPtr<FWorldObject>>& OutWorldObjects) = 0;
+	/** Sets whether or not the UWorld represented by InWorldObject is traceable */
+	virtual void SetWorldTraceability(TSharedRef<FWorldObject> InWorldObject, bool bState) = 0;
+
+	/** Returns all currently available World Trace filters */
+	virtual const TArray<TSharedPtr<IWorldTraceFilter>>& GetWorldFilters() = 0;
+};
\ No newline at end of file
diff --git a/Engine/Plugins/Developer/TraceSourceFiltering/Source/SourceFilteringEditor/Private/SFilterObjectWidget.cpp b/Engine/Plugins/Developer/TraceSourceFiltering/Source/SourceFilteringEditor/Private/SFilterObjectWidget.cpp
new file mode 100644
--- /dev/null
+++ b/Engine/Plugins/Developer/TraceSourceFiltering/Source/SourceFilteringEditor/Private/SFilterObjectWidget.cpp
@@ -0,0 +1,70 @@
+// Copyright Epic Games, Inc. All Rights Reserved.
+
+#include "SFilterObjectWidget.h"
+#include "Widgets/Views/STableRow.h"
+#include "Widgets/Views/STableViewBase.h"
+#include "Widgets/SBoxPanel.h"
+#include "Widgets/Text/STextBlock.h"
+#include "Widgets/Input/SButton.h"
+#include "EditorStyleSet.h"
+#include "Styling/CoreStyle.h"
+#include "Types/SlateEnums.h"
+#include "IFilterObject.h"
+#include "FilterDragDropOperation.h"
+#include "Input/Reply.h"
+#include "Widgets/Layout/SWrapBox.h"
+
+void SFilterObjectRowWidget::Construct(const FArguments& InArgs, const TSharedRef<STableViewBase>& InOwnerTableView, TSharedPtr<IFilterObject> InObject)
+{
+	Object = InObject;
+
+	TSharedRef<SWrapBox> WrapBox = SNew(SWrapBox)
+		.UseAllottedWidth(true);
+
+	TAttribute<bool> FilterEnabled = TAttribute<bool>::Create(TAttribute<bool>::FGetter::CreateSP(Object.Get(), &IFilterObject::IsFilterEnabled));
+	WrapBox->SetEnabled(FilterEnabled);
+	
+	Object->MakeWidget(WrapBox);
+
+	STableRow<TSharedPtr<IFilterObject>>::Construct(
+		STableRow<TSharedPtr<IFilterObject>>::FArguments()
+		.Padding(4.0f)
+		.Content()
+		[
+			WrapBox
+		], InOwnerTableView);
+	
+	TAttribute<FText> ToolTipText = TAttribute<FText>::Create(TAttribute<FText>::FGetter::CreateSP(Object.Get(), &IFilterObject::GetToolTipText));
+	SetToolTipText(ToolTipText);
+}
+
+FReply SFilterObjectRowWidget::OnDragDetected(const FGeometry& MyGeometry, const FPointerEvent& MouseEvent)
+{
+	if (MouseEvent.IsMouseButtonDown(EKeys::LeftMouseButton))
+	{
+		return FReply::Handled().BeginDragDrop(FFilterDragDropOp::New(Object->AsShared()));
+	}
+
+	return FReply::Unhandled();
+}
+
+void SFilterObjectRowWidget::OnDragEnter(FGeometry const& MyGeometry, FDragDropEvent const& DragDropEvent)
+{
+	Object->HandleDragEnter(DragDropEvent);
+}
+
+void SFilterObjectRowWidget::OnDragLeave(FDragDropEvent const& DragDropEvent)
+{
+	Object->HandleDragLeave(DragDropEvent);
+}
+
+FReply SFilterObjectRowWidget::OnDrop(const FGeometry& MyGeometry, const FDragDropEvent& DragDropEvent)
+{
+	return Object->HandleDrop(DragDropEvent);
+}
+
+FReply SFilterObjectRowWidget::OnDragOver(const FGeometry& MyGeometry, const FDragDropEvent& DragDropEvent)
+{
+	return Object->HandleDragEnter(DragDropEvent);
+}
+
diff --git a/Engine/Plugins/Developer/TraceSourceFiltering/Source/SourceFilteringEditor/Private/SFilterObjectWidget.h b/Engine/Plugins/Developer/TraceSourceFiltering/Source/SourceFilteringEditor/Private/SFilterObjectWidget.h
new file mode 100644
--- /dev/null
+++ b/Engine/Plugins/Developer/TraceSourceFiltering/Source/SourceFilteringEditor/Private/SFilterObjectWidget.h
@@ -0,0 +1,32 @@
+// Copyright Epic Games, Inc. All Rights Reserved.
+
+#pragma once
+
+#include "Templates/SharedPointer.h"
+#include "Widgets/Views/STableRow.h"
+#include "Widgets/Views/STableViewBase.h"
+#include "Widgets/Views/SListView.h"
+
+class IFilterObject;
+class STableViewBase;
+
+class SFilterObjectRowWidget : public STableRow<TSharedPtr<IFilterObject>>
+{
+	SLATE_BEGIN_ARGS(SFilterObjectRowWidget) {}
+	SLATE_END_ARGS()
+
+	void Construct(const FArguments& InArgs, const TSharedRef<STableViewBase>& InOwnerTableView, TSharedPtr<IFilterObject> InObject);
+	virtual ~SFilterObjectRowWidget() {}
+
+	/** Begin STableRow<TSharedPtr<IFilterObject>> overrides */
+	virtual FReply OnDragDetected(const FGeometry& MyGeometry, const FPointerEvent& MouseEvent) override;
+	virtual void OnDragEnter(FGeometry const& MyGeometry, FDragDropEvent const& DragDropEvent) override;
+	virtual void OnDragLeave(FDragDropEvent const& DragDropEvent) override;
+	virtual FReply OnDrop(const FGeometry& MyGeometry, const FDragDropEvent& DragDropEvent) override;
+	virtual FReply OnDragOver(const FGeometry& MyGeometry, const FDragDropEvent& DragDropEvent) override;
+	/** End STableRow<TSharedPtr<IFilterObject>> overrides */
+
+protected:	
+	/** Filter object this row represents */
+	TSharedPtr<IFilterObject> Object;	
+};
diff --git a/Engine/Plugins/Developer/TraceSourceFiltering/Source/SourceFilteringEditor/Private/SSourceFilteringTreeview.cpp b/Engine/Plugins/Developer/TraceSourceFiltering/Source/SourceFilteringEditor/Private/SSourceFilteringTreeview.cpp
new file mode 100644
--- /dev/null
+++ b/Engine/Plugins/Developer/TraceSourceFiltering/Source/SourceFilteringEditor/Private/SSourceFilteringTreeview.cpp
@@ -0,0 +1,69 @@
+// Copyright Epic Games, Inc. All Rights Reserved.
+
+#include "SSourceFilteringTreeview.h"
+#include "FilterDragDropOperation.h"
+
+#define LOCTEXT_NAMESPACE "SSourceFilteringTreeView"
+
+void SSourceFilteringTreeView::Construct(const FArguments& InArgs, TSharedRef<STraceSourceFilteringWidget> InOwner)
+{
+	Owner = InOwner;
+	STreeView::Construct(InArgs);
+}
+
+FReply SSourceFilteringTreeView::OnDragOver(const FGeometry& MyGeometry, const FDragDropEvent& DragDropEvent)
+{
+	TSharedPtr<FFilterDragDropOp> FilterDragOp = DragDropEvent.GetOperationAs<FFilterDragDropOp>();
+
+	if (FilterDragOp.IsValid())
+	{
+		if (FilterDragOp->FilterObject.IsValid() && !Owner.Pin()->FilterObjects.Contains(FilterDragOp->FilterObject))
+		{
+			FilterDragOp->SetIconText(FText::FromString(TEXT("\xf146")));
+			FilterDragOp->SetText(LOCTEXT("RemoveFilterFromSetDragOpLabel", "Remove Filter from Filter Set"));
+		}
+	}
+
+	return FReply::Unhandled();
+}
+
+void SSourceFilteringTreeView::OnDragEnter(const FGeometry& MyGeometry, const FDragDropEvent& DragDropEvent)
+{
+	TSharedPtr<FFilterDragDropOp> FilterDragOp = DragDropEvent.GetOperationAs<FFilterDragDropOp>();
+
+	if (FilterDragOp.IsValid())
+	{
+		if (FilterDragOp->FilterObject.IsValid() && !Owner.Pin()->FilterObjects.Contains(FilterDragOp->FilterObject))
+		{
+			FilterDragOp->SetIconText(FText::FromString(TEXT("\xf146")));
+			FilterDragOp->SetText(LOCTEXT("RemoveFilterFromSetDragOpLabel", "Remove Filter from Filter Set"));
+		}
+	}
+}
+
+void SSourceFilteringTreeView::OnDragLeave(const FDragDropEvent& DragDropEvent)
+{
+	TSharedPtr<FFilterDragDropOp> FilterDragOp = DragDropEvent.GetOperationAs<FFilterDragDropOp>();
+	if (FilterDragOp.IsValid())
+	{
+		FilterDragOp->Reset();
+	}
+}
+
+FReply SSourceFilteringTreeView::OnDrop(const FGeometry& MyGeometry, const FDragDropEvent& DragDropEvent)
+{
+	TSharedPtr<FFilterDragDropOp> FilterDragOp = DragDropEvent.GetOperationAs<FFilterDragDropOp>();
+
+	if (FilterDragOp.IsValid())
+	{
+		if (FilterDragOp->FilterObject.IsValid() && !Owner.Pin()->FilterObjects.Contains(FilterDragOp->FilterObject))
+		{
+			Owner.Pin()->SessionFilterService->MakeTopLevelFilter(FilterDragOp->FilterObject.Pin()->AsShared());
+			return FReply::Handled();
+		}
+	}
+
+	return FReply::Unhandled();
+}
+
+#undef LOCTEXT_NAMESPACE // "SSourceFilteringTreeView"
\ No newline at end of file
diff --git a/Engine/Plugins/Developer/TraceSourceFiltering/Source/SourceFilteringEditor/Private/SSourceFilteringTreeview.h b/Engine/Plugins/Developer/TraceSourceFiltering/Source/SourceFilteringEditor/Private/SSourceFilteringTreeview.h
new file mode 100644
--- /dev/null
+++ b/Engine/Plugins/Developer/TraceSourceFiltering/Source/SourceFilteringEditor/Private/SSourceFilteringTreeview.h
@@ -0,0 +1,26 @@
+// Copyright Epic Games, Inc. All Rights Reserved.
+
+#pragma once
+
+#include "Templates/SharedPointer.h"
+#include "Widgets/Views/STreeView.h"
+
+#include "STraceSourceFilteringWidget.h"
+
+/** Simple treeview derived widget , main purpose is to implement drag-and-drop capabilities */
+class SSourceFilteringTreeView : public STreeView<TSharedPtr<IFilterObject>>
+{
+public:
+	void Construct(const FArguments& InArgs, TSharedRef<STraceSourceFilteringWidget> InOwner);
+
+protected:
+	/** Begin STreeView<TSharedPtr<IFilterObject>> overrides */
+	virtual FReply OnDragOver(const FGeometry& MyGeometry, const FDragDropEvent& DragDropEvent) override;
+	virtual void OnDragEnter(const FGeometry& MyGeometry, const FDragDropEvent& DragDropEvent) override;
+	virtual void OnDragLeave(const FDragDropEvent& DragDropEvent) override;
+	virtual FReply OnDrop(const FGeometry& MyGeometry, const FDragDropEvent& DragDropEvent) override;
+	/** End STreeView<TSharedPtr<IFilterObject>> overrides */
+
+protected:
+	TWeakPtr<STraceSourceFilteringWidget> Owner;
+};
\ No newline at end of file
diff --git a/Engine/Plugins/Developer/TraceSourceFiltering/Source/SourceFilteringEditor/Private/STraceSourceFilteringWidget.cpp b/Engine/Plugins/Developer/TraceSourceFiltering/Source/SourceFilteringEditor/Private/STraceSourceFilteringWidget.cpp
new file mode 100644
--- /dev/null
+++ b/Engine/Plugins/Developer/TraceSourceFiltering/Source/SourceFilteringEditor/Private/STraceSourceFilteringWidget.cpp
@@ -0,0 +1,767 @@
+// Copyright Epic Games, Inc. All Rights Reserved.
+
+#include "STraceSourceFilteringWidget.h"
+#include "EditorStyleSet.h"
+
+#include "UObject/SoftObjectPtr.h"
+#include "UObject/SoftObjectPath.h"
+#include "UObject/Class.h"
+#include "UObject/UObjectGlobals.h"
+#include "UObject/Package.h"
+#include "Widgets/Input/SComboButton.h"
+#include "Algo/Transform.h"
+#include "Widgets/Layout/SSeparator.h"
+#include "Framework/MultiBox/MultiBoxBuilder.h"
+#include "Widgets/Images/SThrobber.h"
+#include "Modules/ModuleManager.h"
+#include "TraceServices/ITraceServicesModule.h"
+#include "TraceServices/SessionService.h"
+#include "UObject/ObjectMacros.h"
+#include "Misc/EnumRange.h"
+
+#include "Insights/IUnrealInsightsModule.h"
+#include "Trace/StoreClient.h"
+
+#include "SSourceFilteringTreeview.h"
+#include "SourceFilterStyle.h"
+#include "SFilterObjectWidget.h"
+#include "TraceSourceFilteringSettings.h"
+#include "SWorldTraceFilteringWidget.h"
+#include "IDataSourceFilterSetInterface.h"
+#include "TreeViewBuilder.h"
+
+#if WITH_EDITOR
+#include "PropertyEditorModule.h"
+#include "IDetailsView.h"
+#include "Engine/Blueprint.h"
+#include "EmptySourceFilter.h"
+#include "Editor.h"
+#include "Subsystems/AssetEditorSubsystem.h"
+#endif // WITH_EDITOR
+
+#define LOCTEXT_NAMESPACE "STraceSourceFilteringWidget"
+
+STraceSourceFilteringWidget::STraceSourceFilteringWidget() : FilteringSettings(nullptr)
+{
+}
+
+STraceSourceFilteringWidget::~STraceSourceFilteringWidget()
+{
+	SaveFilteringSettings();
+}
+
+void STraceSourceFilteringWidget::Construct(const FArguments& InArgs)
+{
+	ITraceServicesModule& TraceServicesModule = FModuleManager::LoadModuleChecked<ITraceServicesModule>("TraceServices");
+
+#if WITH_EDITOR
+	ConstructInstanceDetailsView();
+#endif // WITH_EDITOR
+
+	ConstructTreeview();
+	ConstructMenuBox();
+	
+	ChildSlot
+	[
+		SNew(SBorder)
+		.Padding(4.0f)
+		.BorderImage(FEditorStyle::GetBrush("ToolPanel.GroupBorder"))
+		[
+			SNew(SVerticalBox)
+			+ SVerticalBox::Slot()
+			.AutoHeight()
+			.Padding(FMargin(0.0f, 0.0f, 0.0f, 2.0f))
+			[			
+				SNew(SHorizontalBox)
+				+ SHorizontalBox::Slot()
+				.AutoWidth()
+				[
+					MenuBox->AsShared()
+				]
+				+ SHorizontalBox::Slot()
+				.AutoWidth()
+				[
+					SNew(SThrobber)
+					.Visibility(this, &STraceSourceFilteringWidget::GetThrobberVisibility)
+				]
+			]
+			+ SVerticalBox::Slot()
+			[	
+				SNew(SSplitter)
+				.Orientation(EOrientation::Orient_Vertical)
+				.Style(FSourceFilterStyle::Get(), "SourceFilter.Splitter")
+				.PhysicalSplitterHandleSize(2.0f)
+				+ SSplitter::Slot()
+				.Value(.5f)
+				[
+					SNew(SBox)
+					.Padding(2.0f)
+					[
+						FilterTreeView->AsShared()						
+					]
+				]
+#if WITH_EDITOR
+				+ SSplitter::Slot()
+				.Value(.5f)
+				[	
+					SNew(SSplitter)
+					.Orientation(EOrientation::Orient_Horizontal)
+					.Style(FSourceFilterStyle::Get(), "SourceFilter.Splitter")
+					.PhysicalSplitterHandleSize(2.0f)
+					+ SSplitter::Slot()
+					.Value(.5f)
+					[
+						SNew(SBox)
+						.Padding(2.0f)
+						[
+							FilterInstanceDetailsView->AsShared()
+						]
+					]				
+#endif // WITH_EDITOR
+					+ SSplitter::Slot()
+					.Value(.5f)
+					[	
+						SNew(SBox)
+						.Padding(2.0f)
+						[
+							SAssignNew(WorldFilterWidget, SWorldTraceFilteringWidget)
+						]
+					]
+#if WITH_EDITOR
+				]
+#endif // WITH_EDITOR
+			]	
+		]
+	];
+
+	TAttribute<bool> EnabledAttribute = TAttribute<bool>::Create(TAttribute<bool>::FGetter::CreateSP(this, &STraceSourceFilteringWidget::ShouldWidgetsBeEnabled));
+	FilterTreeView->SetEnabled(EnabledAttribute);
+	MenuBox->SetEnabled(EnabledAttribute);
+
+#if WITH_EDITOR
+	FilterInstanceDetailsView->SetEnabled(EnabledAttribute);
+#endif // WITH_EDITOR
+}
+
+void STraceSourceFilteringWidget::ConstructMenuBox()
+{
+	/** Callback for whenever a Filter class (name) was selected */
+	auto OnFilterClassPicked = [this](FString PickedFilterName)
+	{
+		if (SessionFilterService.Get())
+		{
+			SessionFilterService->AddFilter(PickedFilterName);
+			AddFilterButton->SetIsOpen(false);
+		}
+	};
+
+	SAssignNew(MenuBox, SHorizontalBox)
+	+SHorizontalBox::Slot()
+	.AutoWidth()
+	.VAlign(VAlign_Center)
+	[
+		SAssignNew(AddFilterButton, SComboButton)
+		.Visibility(EVisibility::Visible)
+		.ComboButtonStyle(FSourceFilterStyle::Get(), "SourceFilter.ComboButton")
+		.ForegroundColor(FLinearColor::White)
+		.ContentPadding(0.0f)
+		.OnGetMenuContent(FOnGetContent::CreateLambda([OnFilterClassPicked, this]()
+		{
+			FMenuBuilder MenuBuilder(true, TSharedPtr<FUICommandList>(), SessionFilterService->GetExtender());
+			MenuBuilder.AddWidget(SessionFilterService->GetFilterPickerWidget(FOnFilterClassPicked::CreateLambda(OnFilterClassPicked)), FText::GetEmpty(), true, false);
+			return MenuBuilder.MakeWidget();
+		}))
+		.ButtonContent()
+		[
+			SNew(SHorizontalBox)
+			+ SHorizontalBox::Slot()
+			.AutoWidth()
+			.VAlign(VAlign_Center)
+			[
+				SNew(STextBlock)
+				.TextStyle(FSourceFilterStyle::Get(), "SourceFilter.TextStyle")
+				.Font(FEditorStyle::Get().GetFontStyle("FontAwesome.9"))
+				.Text(FText::FromString(FString(TEXT("\xf0fe"))) /*fa-filter*/)
+			]
+			+ SHorizontalBox::Slot()
+			.AutoWidth()
+			.Padding(2, 0, 0, 0)
+			.VAlign(VAlign_Center)
+			[
+				SNew(STextBlock)
+				.TextStyle(FSourceFilterStyle::Get(), "SourceFilter.TextStyle")
+				.Text(LOCTEXT("FilterMenuLabel", "Add Filter"))
+			]
+		]
+	]
+
+	+SHorizontalBox::Slot()
+	.AutoWidth()
+	.VAlign(VAlign_Center)
+	.Padding(2.0f, 0.0f, 0.0f, 0.0f)
+	[
+		SNew(SComboButton)
+		.Visibility(EVisibility::Visible)
+		.ComboButtonStyle(FSourceFilterStyle::Get(), "SourceFilter.ComboButton")
+		.ForegroundColor(FLinearColor::White)
+		.ContentPadding(0.0f)
+		.OnGetMenuContent(this, &STraceSourceFilteringWidget::OnGetOptionsMenu)
+		.ButtonContent()
+		[
+			SNew(SHorizontalBox)
+			+SHorizontalBox::Slot()
+			.AutoWidth()
+			.VAlign(VAlign_Center)
+			[
+				SNew(STextBlock)
+				.TextStyle(FSourceFilterStyle::Get(), "SourceFilter.TextStyle")
+				.Font(FEditorStyle::Get().GetFontStyle("FontAwesome.9"))
+				.Text(FText::FromString(FString(TEXT("\xf0b0"))) /*fa-filter*/)
+			]
+			+SHorizontalBox::Slot()
+			.AutoWidth()
+			.Padding(2, 0, 0, 0)
+			.VAlign(VAlign_Center)
+			[
+				SNew(STextBlock)
+				.TextStyle(FSourceFilterStyle::Get(), "SourceFilter.TextStyle")
+				.Text(LOCTEXT("OptionMenuLabel", "Options"))
+			]
+		]
+	];
+}
+
+void STraceSourceFilteringWidget::ConstructTreeview()
+{
+	SAssignNew(FilterTreeView, SSourceFilteringTreeView, StaticCastSharedRef<STraceSourceFilteringWidget>(AsShared()))
+	.ItemHeight(20.0f)
+	.TreeItemsSource(&FilterObjects)
+	.OnGetChildren_Lambda([this](TSharedPtr<IFilterObject> InObject, TArray<TSharedPtr<IFilterObject>>& OutChildren)
+	{
+		if (TArray<TSharedPtr<IFilterObject>> * ChildArray = ParentToChildren.Find(InObject))
+		{
+			OutChildren.Append(*ChildArray);
+		}
+	})
+	.OnGenerateRow_Lambda([](TSharedPtr<IFilterObject> InItem, const TSharedRef<STableViewBase>& OwnerTable)
+	{
+		return SNew(SFilterObjectRowWidget, OwnerTable, InItem);
+	})
+#if WITH_EDITOR
+	.OnSelectionChanged_Lambda([this](TSharedPtr<IFilterObject> InItem, ESelectInfo::Type SelectInfo) -> void
+	{
+		if (InItem.IsValid())
+		{
+			UObject* Filter = InItem->GetFilter();
+			FilterInstanceDetailsView->SetObject(Filter);
+		}
+		else
+		{
+			FilterInstanceDetailsView->SetObject(nullptr);
+		}	
+	})
+#endif // WITH_EDITOR
+	.OnContextMenuOpening(this, &STraceSourceFilteringWidget::OnContextMenuOpening);
+}
+
+#if WITH_EDITOR
+void STraceSourceFilteringWidget::ConstructInstanceDetailsView()
+{
+	// Create a property view
+	FPropertyEditorModule& EditModule = FModuleManager::Get().GetModuleChecked<FPropertyEditorModule>("PropertyEditor");
+
+	FDetailsViewArgs DetailsViewArgs(
+		/*bUpdateFromSelection=*/ false,
+		/*bLockable=*/ false,
+		/*bAllowSearch=*/ false,
+		FDetailsViewArgs::HideNameArea,
+		/*bHideSelectionTip=*/ true,
+		/*InNotifyHook=*/ nullptr,
+		/*InSearchInitialKeyFocus=*/ false,
+		/*InViewIdentifier=*/ NAME_None);
+	DetailsViewArgs.DefaultsOnlyVisibility = EEditDefaultsOnlyNodeVisibility::Automatic;
+	DetailsViewArgs.bShowOptions = false;	
+
+	FilterInstanceDetailsView = EditModule.CreateDetailView(DetailsViewArgs);
+}
+#endif // WITH_EDITOR
+
+void STraceSourceFilteringWidget::SetCurrentAnalysisSession(uint32 SessionHandle, TSharedRef<const Trace::IAnalysisSession> AnalysisSession)
+{
+	SessionFilterService = FSourceFilterService::GetFilterServiceForSession(SessionHandle, AnalysisSession);
+
+	WorldFilterWidget->SetSessionFilterService(SessionFilterService);
+
+	RefreshFilteringData();
+}
+
+bool STraceSourceFilteringWidget::HasValidFilterSession() const
+{
+	return SessionFilterService.IsValid();
+}
+
+EVisibility STraceSourceFilteringWidget::GetThrobberVisibility() const
+{
+	if (HasValidFilterSession())
+	{
+		return SessionFilterService->IsActionPending() ? EVisibility::Visible : EVisibility::Hidden;
+	}
+
+	return EVisibility::Hidden;
+}
+
+bool STraceSourceFilteringWidget::ShouldWidgetsBeEnabled() const
+{
+	if (HasValidFilterSession())
+	{
+		return !SessionFilterService->IsActionPending();
+	}
+
+	return false;
+}
+
+void STraceSourceFilteringWidget::RefreshFilteringData()
+{
+	SaveTreeviewState();
+
+	FilterObjects.Empty();
+	ParentToChildren.Empty();
+	FlatFilterObjects.Empty();
+
+	FTreeViewDataBuilder Builder(FilterObjects, ParentToChildren, FlatFilterObjects);
+	SessionFilterService->PopulateTreeView(Builder);
+
+	FilterTreeView->RequestTreeRefresh();
+
+	FilteringSettings = SessionFilterService->GetFilterSettings();
+
+	RestoreTreeviewState();
+}
+
+void STraceSourceFilteringWidget::Tick(const FGeometry& AllottedGeometry, const double InCurrentTime, const float InDeltaTime)
+{
+	if (SessionFilterService.IsValid() )
+	{
+		const FDateTime& Stamp = SessionFilterService->GetTimestamp();
+		
+		if (Stamp != SyncTimestamp)
+		{
+			RefreshFilteringData();
+			SyncTimestamp = Stamp;
+		}
+	}
+	else
+	{
+		IUnrealInsightsModule& InsightsModule = FModuleManager::LoadModuleChecked<IUnrealInsightsModule>("TraceInsights");
+		Trace::FStoreClient* StoreClient = InsightsModule.GetStoreClient();
+		const int32 SessionCount = StoreClient->GetSessionCount();
+
+		if (SessionCount > 0)
+		{
+			const Trace::FStoreClient::FSessionInfo* SessionInfo = StoreClient->GetSessionInfo(SessionCount - 1);
+			if (SessionInfo)
+			{
+				SetCurrentAnalysisSession(SessionInfo->GetTraceId(), InsightsModule.GetAnalysisSession().ToSharedRef());
+			}
+		}
+	}
+}
+
+TSharedRef<SWidget> STraceSourceFilteringWidget::OnGetOptionsMenu()
+{
+	FMenuBuilder Builder(true, TSharedPtr<FUICommandList>(), SessionFilterService->GetExtender());
+	
+	if (FilteringSettings)
+	{
+		Builder.BeginSection(NAME_None, LOCTEXT("VisualizationSectionLabel", "Visualization"));
+		{
+			Builder.AddSubMenu(LOCTEXT("VisualizeLabel", "Visualize"),
+				LOCTEXT("DebugDrawingTooltip", "Sub menu related to Debug Drawing"),
+				FNewMenuDelegate::CreateLambda([this](FMenuBuilder& InMenuBuilder)
+					{
+						InMenuBuilder.AddMenuEntry(
+							LOCTEXT("DrawFilterStateLabel","Actor Filtering"), 
+							LOCTEXT("DrawFilteringStateTooltip", "Draws the bounding box for each filter processed Actor in the world."),
+							FSlateIcon(),
+							FUIAction(
+								FExecuteAction::CreateLambda([this]() -> void
+								{
+									FilteringSettings->bDrawFilteringStates = !FilteringSettings->bDrawFilteringStates;
+									SessionFilterService->UpdateFilterSettings(FilteringSettings);
+								}),
+								FCanExecuteAction::CreateLambda([this]() -> bool
+								{
+									return FilteringSettings != nullptr;
+								}),
+								FGetActionCheckState::CreateLambda([this]() -> ECheckBoxState
+								{
+									return FilteringSettings->bDrawFilteringStates ? ECheckBoxState::Checked : ECheckBoxState::Unchecked;
+								})
+							),
+							NAME_None,
+							EUserInterfaceActionType::ToggleButton
+						);
+
+						InMenuBuilder.AddMenuEntry(
+							LOCTEXT("DrawFilterPassingOnlyLabel","Only Actor(s) passing Filtering"), 
+							LOCTEXT("DrawFilterPassingOnlyTooltip", "Only draws the filtering state for Actors that passsed the filtering state."),
+							FSlateIcon(),
+							FUIAction(
+								FExecuteAction::CreateLambda([this]()
+								{
+									FilteringSettings->bDrawOnlyPassingActors = !FilteringSettings->bDrawOnlyPassingActors;
+									SessionFilterService->UpdateFilterSettings(FilteringSettings);
+								}),
+								FCanExecuteAction::CreateLambda([this]()
+								{
+									return FilteringSettings && FilteringSettings->bDrawFilteringStates;
+								}),
+								FGetActionCheckState::CreateLambda([this]()
+								{
+									return FilteringSettings->bDrawOnlyPassingActors ? ECheckBoxState::Checked : ECheckBoxState::Unchecked;
+								})
+							),
+							NAME_None,
+							EUserInterfaceActionType::ToggleButton
+						);
+
+						InMenuBuilder.AddMenuEntry(
+							LOCTEXT("DrawNonPassingFiltersLabel","Draw non passing Filters"), 
+							LOCTEXT("DrawNonPassingFiltersTooltip", "Draws the Filters that caused an Actor to be filtered out."),
+							FSlateIcon(),
+							FUIAction(
+								FExecuteAction::CreateLambda([this]()
+								{
+									FilteringSettings->bDrawFilterDescriptionForRejectedActors = !FilteringSettings->bDrawFilterDescriptionForRejectedActors;
+									SessionFilterService->UpdateFilterSettings(FilteringSettings);
+								}),
+								FCanExecuteAction::CreateLambda([this]()
+								{
+									return FilteringSettings && FilteringSettings->bDrawFilteringStates;
+								}),
+								FGetActionCheckState::CreateLambda([this]()
+								{
+									return FilteringSettings->bDrawFilterDescriptionForRejectedActors ? ECheckBoxState::Checked : ECheckBoxState::Unchecked;
+								})
+							),
+							NAME_None,
+							EUserInterfaceActionType::ToggleButton
+						);
+					
+					}
+				),
+				false,
+				FSlateIcon(),
+				false
+			);
+		}
+		Builder.EndSection();
+	}
+	
+
+	
+	const FName SectionName = TEXT("FilterOptionsMenu");
+	Builder.BeginSection(SectionName, LOCTEXT("FiltersSectionLabel", "Filters"));
+	{
+		Builder.AddMenuEntry(
+			LOCTEXT("ResetFiltersLabel","Reset Filters"), 
+			LOCTEXT("ResetFiltersTooltip", "Removes all currently set filters."),
+			FSlateIcon(),
+			FUIAction(
+				FExecuteAction::CreateLambda([this]()
+				{
+					SessionFilterService->ResetFilters();
+				}),
+				FCanExecuteAction::CreateLambda([this]()
+				{
+					return FilterObjects.Num() > 0;
+				})
+			)
+		);	
+	}
+	Builder.EndSection();
+
+	return Builder.MakeWidget();	
+}
+
+
+TSharedPtr<SWidget> STraceSourceFilteringWidget::OnContextMenuOpening()
+{
+	/** Selection information */
+	TArray<TSharedPtr<IFilterObject>> FilterSelection;
+	FilterTreeView->GetSelectedItems(FilterSelection);
+	
+	FMenuBuilder MenuBuilder(true, TSharedPtr<const FUICommandList>(), SessionFilterService->GetExtender());
+
+	if (FilterSelection.Num() > 0)
+	{
+		bool bSelectionContainsFilterSet = false;
+		bool bSelectionContainsBPFilter = false;
+		bool bSelectionContainsEmptyFilter = false;
+		bool bSelectionContainsNonEmptyFilter = false;
+		const bool bMultiSelection = FilterSelection.Num() > 1;
+
+		/** Gather information about current filter selection set */
+		for (const TSharedPtr<IFilterObject>& Filter : FilterSelection)
+		{
+			const UObject* FilterObject = Filter->GetFilter();
+			if (ensure(FilterObject))
+			{
+				if (const IDataSourceFilterSetInterface* SetInterface = Cast<IDataSourceFilterSetInterface>(FilterObject))
+				{
+					bSelectionContainsFilterSet = true;
+				}
+
+#if WITH_EDITOR
+				if (const UEmptySourceFilter* EmptyFilter = Cast<UEmptySourceFilter>(FilterObject))
+				{
+					bSelectionContainsEmptyFilter = true;
+				}
+				else
+				{
+					bSelectionContainsNonEmptyFilter = true;
+				}
+#else
+				if (const UTraceDataSourceFilter* TraceFilter = Cast<UTraceDataSourceFilter>(FilterObject))
+				{
+					const bool bIsEmptyFilter = (TraceFilter->ClassName == TEXT("EmptySourceFilter"));
+					bSelectionContainsEmptyFilter |= bIsEmptyFilter;
+					bSelectionContainsNonEmptyFilter |= !bIsEmptyFilter;
+				}
+#endif // WITH_EDITOR
+
+				if (FilterObject->GetClass()->ClassGeneratedBy != nullptr)
+				{
+					bSelectionContainsBPFilter = true;
+				}
+			}
+		}
+
+#if WITH_EDITOR
+		if (bSelectionContainsBPFilter)
+		{
+			MenuBuilder.BeginSection(NAME_None, LOCTEXT("BlueprintFilterSectionLabel", "Blueprint Filter"));
+			{
+				MenuBuilder.AddMenuEntry(
+					LOCTEXT("OpenFilterLabel", "Open Filter Blueprint"),
+					LOCTEXT("OpenFilterTooltip", "Opens this Filter's Blueprint."),
+					FSlateIcon(FEditorStyle::GetStyleSetName(), "ClassIcon.Blueprint"),
+					FUIAction(
+						FExecuteAction::CreateLambda([FilterSelection]()
+						{
+							for (const TSharedPtr<IFilterObject>& FilterObject : FilterSelection)
+							{
+								const UObject* FilterUObject = FilterObject->GetFilter();
+								if (ensure(FilterUObject))
+								{
+									if (UBlueprint* Blueprint = Cast<UBlueprint>(FilterUObject->GetClass()->ClassGeneratedBy))
+									{
+										GEditor->GetEditorSubsystem<UAssetEditorSubsystem>()->OpenEditorForAsset(Blueprint);
+									}
+								}
+							}
+						})
+					)
+				);
+			}
+			MenuBuilder.EndSection();
+
+		}
+#endif // WITH_EDITOR
+
+		/** Single selection of a filter set */
+		if (bSelectionContainsFilterSet && !bMultiSelection)
+		{
+			auto AddFilterToSet = [this, FilterSelection](FString ClassName)
+			{
+				if (SessionFilterService.Get())
+				{
+					SessionFilterService->AddFilterToSet(FilterSelection[0]->AsShared(), ClassName);
+				}
+
+				FSlateApplication::Get().DismissAllMenus();
+			};
+
+			MenuBuilder.BeginSection(NAME_None, LOCTEXT("FilterSetContextMenuLabel", "Filter Set"));
+			{
+				MenuBuilder.AddSubMenu(
+					LOCTEXT("AddFilterToSetLabel", "Add Filter"),
+					LOCTEXT("AddFilterToSetTooltip", "Adds a filter to this Filtering Set."),
+					FNewMenuDelegate::CreateLambda([this, AddFilterToSet](FMenuBuilder& InSubMenuBuilder)
+					{
+						InSubMenuBuilder.AddWidget(
+							SessionFilterService->GetFilterPickerWidget(FOnFilterClassPicked::CreateLambda(AddFilterToSet)),
+							FText::GetEmpty(),
+							true
+						);
+					})
+				);
+			}
+			MenuBuilder.EndSection();
+		}
+
+		MenuBuilder.BeginSection(NAME_None, LOCTEXT("FiltersContextMenuLabel", "Filter"));
+		{
+			MenuBuilder.AddMenuEntry(
+				LOCTEXT("EnableFilterLabel", "Filter Enabled"),
+				LOCTEXT("ToggleFilterTooltips", "Sets whether or not this Filter should be considered when applying the set of filters"),
+				FSlateIcon(),
+				FUIAction(
+					FExecuteAction::CreateLambda([FilterSelection, this]()
+					{
+						bool bEnabled = false;
+						bool bDisabled = false;
+
+						for (const TSharedPtr<IFilterObject>& FilterObject : FilterSelection)
+						{
+							bEnabled |= FilterObject->IsFilterEnabled();
+							bDisabled |= !FilterObject->IsFilterEnabled();
+						}
+
+						for (const TSharedPtr<IFilterObject>& FilterObject : FilterSelection)
+						{
+							SessionFilterService->SetFilterState(FilterObject->AsShared(), (bEnabled && bDisabled) || bDisabled);
+						}
+					}),
+					FCanExecuteAction(),
+					FGetActionCheckState::CreateLambda([FilterSelection]()
+					{
+						bool bEnabled = false;
+						bool bDisabled = false;
+
+						for (const TSharedPtr<IFilterObject>& Filter : FilterSelection)
+						{
+							bEnabled |= Filter->IsFilterEnabled();
+							bDisabled |= !Filter->IsFilterEnabled();
+						}
+
+						return bEnabled ? (bDisabled ? ECheckBoxState::Undetermined : ECheckBoxState::Checked) : ECheckBoxState::Unchecked;
+					})
+				),
+				NAME_None,
+				EUserInterfaceActionType::Check
+			);
+
+			MenuBuilder.AddMenuEntry(
+				LOCTEXT("RemoveFilterLabel", "Remove Filter"),
+				LOCTEXT("RemoveFilterTooltip", "Removes this Filter from the filtering set."),
+				FSlateIcon(),
+				FUIAction(
+					FExecuteAction::CreateLambda([this, FilterSelection]()
+					{
+						for (const TSharedPtr<IFilterObject>& FilterObject : FilterSelection)
+						{
+							SessionFilterService->RemoveFilter(FilterObject->AsShared());
+						}
+					})
+				)
+			);
+		}
+		MenuBuilder.EndSection();
+
+		/** Single selection of a valid Filter instance */
+		if (!bMultiSelection && bSelectionContainsNonEmptyFilter)
+		{
+			MenuBuilder.BeginSection(NAME_None, LOCTEXT("AddFilterSetSectionLabel", "Add Filter Set"));
+			{
+				FText LabelTextFormat = LOCTEXT("MakeFilterSetLabel", "{0}");
+				FText ToolTipTextFormat = LOCTEXT("MakeFilterSetTooltip", "Creates a new filter set, containing this filter, with the {0} operator");
+
+				const UEnum* EnumPtr = FindObject<UEnum>(ANY_PACKAGE, TEXT("EFilterSetMode"), true);
+				for (EFilterSetMode Mode : TEnumRange<EFilterSetMode>())
+				{
+					FText ModeText = EnumPtr->GetDisplayNameTextByValue((int64)Mode);
+					MenuBuilder.AddMenuEntry(
+						FText::Format(LabelTextFormat, ModeText),
+						FText::Format(ToolTipTextFormat, ModeText),
+						FSlateIcon(),
+						FUIAction(
+							FExecuteAction::CreateLambda([this, Mode, FilterSelection]()
+							{
+								SessionFilterService->MakeFilterSet(FilterSelection[0]->AsShared(), Mode);
+							})
+						)
+					);
+				}
+			}
+			MenuBuilder.EndSection();
+		}
+	}
+	else
+	{
+		auto OnFilterClassPicked = [this](FString PickedFilterName)
+		{
+			if (SessionFilterService.Get())
+			{
+				SessionFilterService->AddFilter(PickedFilterName);
+				AddFilterButton->SetIsOpen(false);
+			}
+		};
+		
+		MenuBuilder.AddWidget(SessionFilterService->GetFilterPickerWidget(FOnFilterClassPicked::CreateLambda(OnFilterClassPicked)), FText::GetEmpty(), true, false);
+	}
+
+	return MenuBuilder.MakeWidget();
+}
+
+void STraceSourceFilteringWidget::SaveFilteringSettings()
+{
+	if (FilteringSettings)
+	{
+		FilteringSettings->SaveConfig();
+	}
+}
+
+void STraceSourceFilteringWidget::SaveTreeviewState()
+{
+	if (FilterTreeView)
+	{
+		ensure(ExpandedFilters.Num() == 0);
+		TSet<TSharedPtr<IFilterObject>> TreeviewExpandedObjects;
+		FilterTreeView->GetExpandedItems(TreeviewExpandedObjects);
+		Algo::Transform(TreeviewExpandedObjects, ExpandedFilters, [](TSharedPtr<IFilterObject> Object)
+		{
+			return GetTypeHash(Object->GetFilter());
+		});
+
+		ensure(SelectedFilters.Num() == 0);
+		TArray<TSharedPtr<IFilterObject>> TreeviewSelectedObjects;
+		FilterTreeView->GetSelectedItems(TreeviewSelectedObjects);
+		Algo::Transform(TreeviewSelectedObjects, SelectedFilters, [](TSharedPtr<IFilterObject> Object)
+		{
+			return GetTypeHash(Object->GetFilter());
+		});
+
+	}
+}
+
+void STraceSourceFilteringWidget::RestoreTreeviewState()
+{
+	if (FilterTreeView)
+	{
+		FilterTreeView->ClearExpandedItems();
+		for (TSharedPtr<IFilterObject> FilterObject : FlatFilterObjects)
+		{
+			if (ExpandedFilters.Contains(GetTypeHash(FilterObject->GetFilter())))
+			{
+				FilterTreeView->SetItemExpansion(FilterObject, true);
+			}
+		}
+
+		ExpandedFilters.Empty();
+
+		FilterTreeView->ClearSelection();
+		for (TSharedPtr<IFilterObject> FilterObject : FlatFilterObjects)
+		{
+			if (SelectedFilters.Contains(GetTypeHash(FilterObject->GetFilter())))
+			{
+				FilterTreeView->SetItemSelection(FilterObject, true);
+			}
+		}
+
+		SelectedFilters.Empty();
+	}
+}
+
+
+#undef LOCTEXT_NAMESPACE // "STraceSourceFilteringWidget"
\ No newline at end of file
diff --git a/Engine/Plugins/Developer/TraceSourceFiltering/Source/SourceFilteringEditor/Private/STraceSourceFilteringWidget.h b/Engine/Plugins/Developer/TraceSourceFiltering/Source/SourceFilteringEditor/Private/STraceSourceFilteringWidget.h
new file mode 100644
--- /dev/null
+++ b/Engine/Plugins/Developer/TraceSourceFiltering/Source/SourceFilteringEditor/Private/STraceSourceFilteringWidget.h
@@ -0,0 +1,107 @@
+// Copyright Epic Games, Inc. All Rights Reserved.
+
+#pragma once
+
+#include "Templates/SharedPointer.h"
+#include "Widgets/Views/STreeView.h"
+
+#include "SourceFilterService.h"
+#include "ISessionSourceFilterService.h"
+#include "IFilterObject.h"
+
+#if WITH_EDITOR
+#include "TraceSourceFilteringSettings.h"
+#endif // WITH_EDITOR
+
+class IDetailsView;
+class SWorldTraceFilteringWidget;
+class SComboButton;
+class SHorizontalBox;
+
+class STraceSourceFilteringWidget : public SCompoundWidget
+{
+	friend class SSourceFilteringTreeView;
+public:
+	/** Default constructor. */
+	STraceSourceFilteringWidget();
+
+	/** Virtual destructor. */
+	virtual ~STraceSourceFilteringWidget();
+
+	SLATE_BEGIN_ARGS(STraceSourceFilteringWidget) {}
+	SLATE_END_ARGS()
+
+	/** Constructs this widget. */
+	void Construct(const FArguments& InArgs);
+
+	/** Begin SCompoundWidget overrides */
+	virtual void Tick(const FGeometry& AllottedGeometry, const double InCurrentTime, const float InDeltaTime) override;
+	/** End SCompoundWidget overrides */
+	
+protected:
+	void ConstructMenuBox();
+	void ConstructTreeview();
+#if WITH_EDITOR
+	void ConstructInstanceDetailsView();
+#endif // WITH_EDITOR
+
+	TSharedRef<SWidget> OnGetOptionsMenu();
+	TSharedPtr<SWidget> OnContextMenuOpening();
+
+	/** Callback for whenever a different analysis session (store) has been retrieved */
+	void SetCurrentAnalysisSession(uint32 SessionHandle, TSharedRef<const Trace::IAnalysisSession> AnalysisSession);
+
+	/** Returns whether or not a valid ISessionSourceFilterService is available / set */
+	bool HasValidFilterSession() const;
+
+	/** Returns visibility state for SThrobber, used to indicate pending filter session request */
+	EVisibility GetThrobberVisibility() const;
+
+	/** Returns whether or not the contained widgets should be enabled, determined by having a valid session and no pending request */
+	bool ShouldWidgetsBeEnabled() const;
+
+	/** Refreshes the filtering data and state, using SessionFilterService, represented by this widget */
+	void RefreshFilteringData();
+	
+	/** Save current UTraceSourceFilteringSettings state to INI files */
+	void SaveFilteringSettings();
+
+	/** Saves and restores the current selection and expansion state of FilterTreeView*/
+	void SaveTreeviewState();
+	void RestoreTreeviewState();
+protected:
+#if WITH_EDITOR	
+	/** Details view, used for displaying selected Filter UProperties */
+	TSharedPtr<IDetailsView> FilterInstanceDetailsView;	
+#endif // WITH_EDITOR
+	
+	/** Slate widget used to add filter instances to the session */
+	TSharedPtr<SComboButton> AddFilterButton;
+
+	/** Slate widget containing the Add Filter and Options widgets, used for disabling/enabling according to the session state */
+	TSharedPtr<SHorizontalBox> MenuBox;
+
+	/** Filter session instance, used to retrieve data and communicate with connected application */
+	TSharedPtr<ISessionSourceFilterService> SessionFilterService;
+
+	/** Panel used for filtering UWorld's traceability on the connected filter session */
+	TSharedPtr<SWorldTraceFilteringWidget> WorldFilterWidget;
+
+	/** Treeview used to display all currently represented Filters */
+	TSharedPtr<STreeView<TSharedPtr<IFilterObject>>> FilterTreeView;
+
+	/** Data used to populate the Filter Treeview */
+	TArray<TSharedPtr<IFilterObject>> FilterObjects;
+	TMap<TSharedPtr<IFilterObject>, TArray<TSharedPtr<IFilterObject>>> ParentToChildren;
+	TArray<TSharedPtr<IFilterObject>> FlatFilterObjects;
+	   
+	/** Timestamp at which the treeview data was last retrieved from SessionFilterService */
+	FDateTime SyncTimestamp;
+
+	/** Used to store hash values of, treeview expanded and selected, filter instances when refreshing the treeview data */
+	TArray<int32> ExpandedFilters;	
+	TArray<int32> SelectedFilters;
+
+	/** Cached pointer to Filtering Settings retrieved from SessionFilterService */
+	UTraceSourceFilteringSettings* FilteringSettings;
+};
diff --git a/Engine/Plugins/Developer/TraceSourceFiltering/Source/SourceFilteringEditor/Private/SWorldObjectWidget.cpp b/Engine/Plugins/Developer/TraceSourceFiltering/Source/SourceFilteringEditor/Private/SWorldObjectWidget.cpp
new file mode 100644
--- /dev/null
+++ b/Engine/Plugins/Developer/TraceSourceFiltering/Source/SourceFilteringEditor/Private/SWorldObjectWidget.cpp
@@ -0,0 +1,51 @@
+// Copyright Epic Games, Inc. All Rights Reserved.
+
+#include "SWorldObjectWidget.h"
+#include "Widgets/Input/SCheckBox.h"
+#include "Widgets/SBoxPanel.h"
+#include "WorldObject.h"
+#include "Widgets/Views/STableRow.h"
+
+#include "ISessionSourceFilterService.h"
+
+void SWorldObjectRowWidget::Construct(const FArguments& InArgs, const TSharedRef<STableViewBase>& InOwnerTableView, TSharedPtr<FWorldObject> InWorldObject, TSharedPtr<ISessionSourceFilterService> InFilterService)
+{
+	WorldObject = InWorldObject;
+	FilterService = InFilterService;
+
+	STableRow<TSharedPtr<FWorldObject>>::Construct(
+		STableRow<TSharedPtr<FWorldObject>>::FArguments()
+		.Padding(0.0f)
+		.Content()
+		[
+			SNew(SHorizontalBox)
+			// Name 
+			+ SHorizontalBox::Slot()
+			.FillWidth(0.5f)
+			.Padding(2, 0, 0, 0)
+			.VAlign(VAlign_Center)
+			[
+				SNew(STextBlock)
+				.Text(WorldObject->GetDisplayText())
+			]
+			// Traceability checkbox 
+			+ SHorizontalBox::Slot()
+			.AutoWidth()
+			.VAlign(VAlign_Center)
+			.Padding(0, 0, 2.f, 0)
+			[
+				/** Called when the checked state has changed */
+				SNew(SCheckBox)
+				.IsChecked_Lambda([this]() -> ECheckBoxState
+				{
+					return WorldObject->CanOutputData() ? ECheckBoxState::Checked : ECheckBoxState::Unchecked;
+				})
+				.OnCheckStateChanged_Lambda([this](ECheckBoxState State)
+				{
+					FilterService->SetWorldTraceability(WorldObject->AsShared(), State == ECheckBoxState::Checked);
+				})
+			]
+		],
+		InOwnerTableView
+	);
+}
diff --git a/Engine/Plugins/Developer/TraceSourceFiltering/Source/SourceFilteringEditor/Private/SWorldObjectWidget.h b/Engine/Plugins/Developer/TraceSourceFiltering/Source/SourceFilteringEditor/Private/SWorldObjectWidget.h
new file mode 100644
--- /dev/null
+++ b/Engine/Plugins/Developer/TraceSourceFiltering/Source/SourceFilteringEditor/Private/SWorldObjectWidget.h
@@ -0,0 +1,31 @@
+// Copyright Epic Games, Inc. All Rights Reserved.
+
+#pragma once
+
+#include "Templates/SharedPointer.h"
+#include "Widgets/DeclarativeSyntaxSupport.h"
+#include "Widgets/Views/STableRow.h"
+#include "Widgets/Views/STableViewBase.h"
+#include "Widgets/Views/SListView.h"
+
+class FWorldObject;
+class STableViewBase;
+class ISessionSourceFilterService;
+class FWorldObject;
+
+/** Table row representing a UWorld instance, as a FWorldObject, within SWorldFilterWidget */
+class SWorldObjectRowWidget : public STableRow<TSharedPtr<FWorldObject>>
+{
+	SLATE_BEGIN_ARGS(SWorldObjectRowWidget) {}
+	SLATE_END_ARGS()
+
+	void Construct(const FArguments& InArgs, const TSharedRef<STableViewBase>& InOwnerTableView, TSharedPtr<FWorldObject> InWorldObject, TSharedPtr<ISessionSourceFilterService> InFilterService);
+	virtual ~SWorldObjectRowWidget() {}
+
+protected:
+	/* The World object that this tree item represents */
+	TSharedPtr<FWorldObject> WorldObject;
+
+	/** Session service that has generated this object */
+	TSharedPtr<ISessionSourceFilterService> FilterService;
+};
diff --git a/Engine/Plugins/Developer/TraceSourceFiltering/Source/SourceFilteringEditor/Private/SWorldTraceFilteringWidget.cpp b/Engine/Plugins/Developer/TraceSourceFiltering/Source/SourceFilteringEditor/Private/SWorldTraceFilteringWidget.cpp
new file mode 100644
--- /dev/null
+++ b/Engine/Plugins/Developer/TraceSourceFiltering/Source/SourceFilteringEditor/Private/SWorldTraceFilteringWidget.cpp
@@ -0,0 +1,143 @@
+// Copyright Epic Games, Inc. All Rights Reserved.
+
+#include "SWorldTraceFilteringWidget.h"
+
+#include "Widgets/Input/SComboButton.h"
+#include "Widgets/SBoxPanel.h"
+#include "Widgets/Text/STextBlock.h"
+#include "EditorStyleSet.h"
+#include "Internationalization/Internationalization.h"
+#include "Algo/Transform.h"
+#include "Widgets/Views/SListView.h"
+
+#include "SWorldObjectWidget.h"
+#include "SourceFilterStyle.h"
+#include "ISessionSourceFilterService.h"
+#include "WorldObject.h"
+
+#define LOCTEXT_NAMESPACE "SWorldFilterWidget"
+
+void SWorldTraceFilteringWidget::Construct(const FArguments& InArgs)
+{
+	ChildSlot
+	[
+		SNew(SBorder)
+		.Padding(4.0f)
+		.BorderImage(FEditorStyle::GetBrush("ToolPanel.GroupBorder"))
+		[
+			SNew(SVerticalBox)
+			+ SVerticalBox::Slot()
+			.AutoHeight()
+			.Padding(FMargin(0.0f, 0.0f, 0.0f, 2.0f))
+			[			
+				SNew(SHorizontalBox)
+				+ SHorizontalBox::Slot()
+				.AutoWidth()
+				.VAlign(VAlign_Center)
+				[
+					SAssignNew(OptionsComboBox, SComboButton)
+					.Visibility(EVisibility::Visible)
+					.ComboButtonStyle(FSourceFilterStyle::Get(), "SourceFilter.ComboButton")
+					.ForegroundColor(FLinearColor::White)
+					.ContentPadding(0.0f)
+					.OnGetMenuContent(this, &SWorldTraceFilteringWidget::OnGetMenuContextMenu)					
+					.ButtonContent()
+					[
+						SNew(SHorizontalBox)
+						+ SHorizontalBox::Slot()
+						.AutoWidth()
+						.VAlign(VAlign_Center)
+						[
+							SNew(STextBlock)
+							.TextStyle(FSourceFilterStyle::Get(), "SourceFilter.TextStyle")
+							.Font(FEditorStyle::Get().GetFontStyle("FontAwesome.9"))
+							.Text(FText::FromString(FString(TEXT("\xf0fe"))) /*fa-filter*/)
+						]
+						+ SHorizontalBox::Slot()
+						.AutoWidth()
+						.Padding(2, 0, 0, 0)
+						.VAlign(VAlign_Center)
+						[
+							SNew(STextBlock)
+							.TextStyle(FSourceFilterStyle::Get(), "SourceFilter.TextStyle")
+							.Text(LOCTEXT("OptionsMenuLabel", "Options"))
+						]
+					]
+				]
+			]
+
+			+ SVerticalBox::Slot()
+			[	
+				SAssignNew(WorldListView, SListView<TSharedPtr<FWorldObject>>)
+				.ItemHeight(20.f)
+				.ListItemsSource(&WorldObjects)
+				.OnGenerateRow_Lambda([this](TSharedPtr<FWorldObject> InItem, const TSharedRef<STableViewBase>& OwnerTable) -> TSharedRef<ITableRow>
+				{
+					return SNew(SWorldObjectRowWidget, OwnerTable, InItem, SessionFilterService);
+				})
+			]	
+		]
+	];
+
+	TAttribute<bool> EnabledAttribute = TAttribute<bool>::Create([this]() -> bool
+	{
+		return SessionFilterService.IsValid() && !SessionFilterService->IsActionPending();
+	});
+
+	OptionsComboBox->SetEnabled(EnabledAttribute);
+	WorldListView->SetEnabled(EnabledAttribute);
+}
+
+void SWorldTraceFilteringWidget::SetSessionFilterService(TSharedPtr<ISessionSourceFilterService> InSessionFilterService)
+{
+	SessionFilterService = InSessionFilterService;
+	RefreshWorldData();
+}
+
+void SWorldTraceFilteringWidget::Tick(const FGeometry& AllottedGeometry, const double InCurrentTime, const float InDeltaTime)
+{
+	if (SessionFilterService.IsValid())
+	{
+		if (SessionFilterService->GetTimestamp() > TimeStamp)
+		{
+			RefreshWorldData();
+			TimeStamp = SessionFilterService->GetTimestamp();
+		}
+	}
+}
+
+void SWorldTraceFilteringWidget::RefreshWorldData()
+{
+	WorldObjects.Empty();
+	SessionFilterService->GetWorldObjects(WorldObjects);
+	WorldListView->RequestListRefresh();
+}
+
+TSharedRef<SWidget> SWorldTraceFilteringWidget::OnGetMenuContextMenu()
+{
+	FMenuBuilder MenuBuilder(true, TSharedPtr<FUICommandList>(), SessionFilterService->GetExtender());
+	const TArray<TSharedPtr<IWorldTraceFilter>>& WorldFilters = SessionFilterService->GetWorldFilters();
+
+	if (WorldFilters.Num())
+	{
+		MenuBuilder.BeginSection(NAME_None, LOCTEXT("WorldFilteringLabel", "World Filtering"));
+		{
+			for (TSharedPtr<IWorldTraceFilter> WorldFilter : WorldFilters)
+			{
+				MenuBuilder.AddSubMenu(WorldFilter->GetDisplayText(),
+					WorldFilter->GetToolTipText(),
+					FNewMenuDelegate::CreateLambda([this, WorldFilter](FMenuBuilder& SubMenuBuilder)
+					{
+						WorldFilter->PopulateMenuBuilder(SubMenuBuilder);
+					}),
+					false
+				);
+			}
+		}
+		MenuBuilder.EndSection();
+	}
+		
+	return MenuBuilder.MakeWidget();
+}
+
+#undef LOCTEXT_NAMESPACE // "SWorldFilterWidget"
\ No newline at end of file
diff --git a/Engine/Plugins/Developer/TraceSourceFiltering/Source/SourceFilteringEditor/Private/SWorldTraceFilteringWidget.h b/Engine/Plugins/Developer/TraceSourceFiltering/Source/SourceFilteringEditor/Private/SWorldTraceFilteringWidget.h
new file mode 100644
--- /dev/null
+++ b/Engine/Plugins/Developer/TraceSourceFiltering/Source/SourceFilteringEditor/Private/SWorldTraceFilteringWidget.h
@@ -0,0 +1,48 @@
+// Copyright Epic Games, Inc. All Rights Reserved.
+
+#pragma once
+
+#include "Widgets/SCompoundWidget.h"
+#include "Widgets/DeclarativeSyntaxSupport.h"
+#include "Widgets/Input/SComboBox.h"
+
+class SComboButton;
+class ISessionSourceFilterService;
+class FWorldObject;
+
+class SWorldTraceFilteringWidget : public SCompoundWidget
+{
+public:
+	/** Default constructor. */
+	SWorldTraceFilteringWidget() {}
+
+	virtual ~SWorldTraceFilteringWidget() {}
+
+	SLATE_BEGIN_ARGS(SWorldTraceFilteringWidget) {}
+	SLATE_END_ARGS()
+
+	/** Begin SCompoundWidget overrides */
+	virtual void Tick(const FGeometry& AllottedGeometry, const double InCurrentTime, const float InDeltaTime) override;
+	/** End SCompoundWidget overrides */
+
+	void Construct(const FArguments& InArgs);
+	void SetSessionFilterService(TSharedPtr<ISessionSourceFilterService> InSessionFilterService);
+protected:
+	void RefreshWorldData();
+	TSharedRef<SWidget> OnGetMenuContextMenu();
+
+protected:
+	/** Session service used to retrieve state and request filtering changes */
+	TSharedPtr<ISessionSourceFilterService> SessionFilterService;
+
+	/** Latest set of retrieve World Objects from SessionFilterService */
+	TArray<TSharedPtr<FWorldObject>> WorldObjects;
+
+	/** Listview widget used to display WorldObjects  */
+	TSharedPtr<SListView<TSharedPtr<FWorldObject>>> WorldListView;
+
+	TSharedPtr<SComboButton> OptionsComboBox;
+
+	/** Cached timestamp, at which WorldObjects was last updated */
+	FDateTime TimeStamp;
+};
\ No newline at end of file
diff --git a/Engine/Plugins/Developer/TraceSourceFiltering/Source/SourceFilteringEditor/Private/SourceFilterService.cpp b/Engine/Plugins/Developer/TraceSourceFiltering/Source/SourceFilteringEditor/Private/SourceFilterService.cpp
new file mode 100644
--- /dev/null
+++ b/Engine/Plugins/Developer/TraceSourceFiltering/Source/SourceFilteringEditor/Private/SourceFilterService.cpp
@@ -0,0 +1,19 @@
+// Copyright Epic Games, Inc. All Rights Reserved.
+
+#include "SourceFilterService.h"
+
+#include "Modules/ModuleManager.h"
+#include "Templates/SharedPointer.h"
+
+#include "TraceServices/Model/AnalysisSession.h"
+
+#if WITH_EDITOR
+#include "EditorSessionSourceFilterService.h"
+#endif
+
+TSharedRef<ISessionSourceFilterService> FSourceFilterService::GetFilterServiceForSession(uint32 InHandle, TSharedRef<const Trace::IAnalysisSession> AnalysisSession)
+{
+	TSharedPtr<ISessionSourceFilterService> NewService = nullptr;		
+	NewService = MakeShareable(new FEditorSessionSourceFilterService());
+	return NewService->AsShared();
+}
\ No newline at end of file
diff --git a/Engine/Plugins/Developer/TraceSourceFiltering/Source/SourceFilteringEditor/Private/SourceFilterService.h b/Engine/Plugins/Developer/TraceSourceFiltering/Source/SourceFilteringEditor/Private/SourceFilterService.h
new file mode 100644
--- /dev/null
+++ b/Engine/Plugins/Developer/TraceSourceFiltering/Source/SourceFilteringEditor/Private/SourceFilterService.h
@@ -0,0 +1,19 @@
+// Copyright Epic Games, Inc. All Rights Reserved.
+
+#pragma once
+
+#include "Containers/Array.h"
+#include "Containers/Map.h"
+#include "ISessionSourceFilterService.h"
+
+namespace Trace
+{
+	class IAnalysisSession;
+}
+
+class FSourceFilterService
+{
+public:
+	/** Retrieves the ISessionSourceFilterService for the provided Trace Session Handle */
+	static TSharedRef<ISessionSourceFilterService> GetFilterServiceForSession(uint32 InHandle, TSharedRef<const Trace::IAnalysisSession> AnalysisSession);
+};
\ No newline at end of file
diff --git a/Engine/Plugins/Developer/TraceSourceFiltering/Source/SourceFilteringEditor/Private/SourceFilterStyle.cpp b/Engine/Plugins/Developer/TraceSourceFiltering/Source/SourceFilteringEditor/Private/SourceFilterStyle.cpp
new file mode 100644
--- /dev/null
+++ b/Engine/Plugins/Developer/TraceSourceFiltering/Source/SourceFilteringEditor/Private/SourceFilterStyle.cpp
@@ -0,0 +1,179 @@
+// Copyright Epic Games, Inc. All Rights Reserved.
+
+#include "SourceFilterStyle.h"
+#include "Styling/SlateStyleRegistry.h"
+#include "Styling/SlateTypes.h"
+#include "Styling/CoreStyle.h"
+
+
+FSlateBrush* FSourceFilterStyle::FilterBrush = nullptr;
+FSlateBrush* FSourceFilterStyle::FilterSetOperationBrush = nullptr;
+FSlateBrush* FSourceFilterStyle::FilterSetBrush = nullptr;
+
+TSharedPtr< FSlateStyleSet > FSourceFilterStyle::StyleSet = nullptr;
+
+FTextBlockStyle FSourceFilterStyle::NormalText;
+
+#define IMAGE_BRUSH( RelativePath, ... ) FSlateImageBrush( StyleSet->RootToContentDir( RelativePath, TEXT(".png") ), __VA_ARGS__ )
+#define BOX_BRUSH( RelativePath, ... ) FSlateBoxBrush( StyleSet->RootToContentDir( RelativePath, TEXT(".png") ), __VA_ARGS__ )
+#define BORDER_BRUSH( RelativePath, ... ) FSlateBorderBrush( StyleSet->RootToContentDir( RelativePath, TEXT(".png") ), __VA_ARGS__ )
+#define DEFAULT_FONT(...) FCoreStyle::GetDefaultFontStyle(__VA_ARGS__)
+
+// Const icon sizes
+static const FVector2D Icon8x8(8.0f, 8.0f);
+static const FVector2D Icon9x19(9.0f, 19.0f);
+static const FVector2D Icon14x14(14.0f, 14.0f);
+static const FVector2D Icon16x16(16.0f, 16.0f);
+static const FVector2D Icon20x20(20.0f, 20.0f);
+static const FVector2D Icon22x22(22.0f, 22.0f);
+static const FVector2D Icon24x24(24.0f, 24.0f);
+static const FVector2D Icon28x28(28.0f, 28.0f);
+static const FVector2D Icon27x31(27.0f, 31.0f);
+static const FVector2D Icon26x26(26.0f, 26.0f);
+static const FVector2D Icon32x32(32.0f, 32.0f);
+static const FVector2D Icon40x40(40.0f, 40.0f);
+static const FVector2D Icon48x48(48.0f, 48.0f);
+static const FVector2D Icon75x82(75.0f, 82.0f);
+static const FVector2D Icon360x32(360.0f, 32.0f);
+static const FVector2D Icon171x39(171.0f, 39.0f);
+static const FVector2D Icon170x50(170.0f, 50.0f);
+static const FVector2D Icon267x140(170.0f, 50.0f);
+
+void FSourceFilterStyle::Initialize()
+{
+	// Only register once
+	if( StyleSet.IsValid() )
+	{
+		return;
+	}
+
+	StyleSet = MakeShareable(new FSlateStyleSet("SourceFilter") );
+	
+	StyleSet->SetContentRoot(FPaths::EngineContentDir() / TEXT("Editor/Slate"));
+	StyleSet->SetCoreContentRoot(FPaths::EngineContentDir() / TEXT("Slate"));
+
+
+	NormalText = FTextBlockStyle()
+		.SetFont(DEFAULT_FONT("Regular", FCoreStyle::RegularTextSize))
+		.SetColorAndOpacity(FSlateColor::UseForeground())
+		.SetShadowOffset(FVector2D::ZeroVector)
+		.SetShadowColorAndOpacity(FLinearColor::Black)
+		.SetHighlightColor(FLinearColor(0.02f, 0.3f, 0.0f))
+		.SetHighlightShape(BOX_BRUSH("Common/TextBlockHighlightShape", FMargin(3.f / 8.f)));
+
+
+	FSourceFilterStyle::FilterBrush = new BOX_BRUSH("Common/FlatButton", 2.0f / 8.0f, FLinearColor(0.25f, 0.25f, 0.25f, 0.9f));
+	StyleSet->Set("SourceFilter.FilterBrush", FSourceFilterStyle::FilterBrush);
+
+	FSourceFilterStyle::FilterSetOperationBrush = new BOX_BRUSH("Common/FlatButton", 2.0f / 8.0f, FLinearColor(0.85f, 0.25f, 0.25f, 0.9f));
+	StyleSet->Set("SourceFilter.FilterSetOperationBrush", FSourceFilterStyle::FilterSetOperationBrush);
+		
+	FSourceFilterStyle::FilterSetBrush = new BOX_BRUSH("Common/FlatButton", 2.0f / 8.0f, FLinearColor(0.125f, 0.125f, 0.125f, 0.9f));
+	StyleSet->Set("SourceFilter.FilterSetBrush", FSourceFilterStyle::FilterSetBrush);
+
+	// Normal button
+	FButtonStyle Button = FButtonStyle()
+		.SetNormal(BOX_BRUSH("Common/Button", FVector2D(32, 32), 8.0f / 32.0f))
+		.SetHovered(BOX_BRUSH("Common/Button_Hovered", FVector2D(32, 32), 8.0f / 32.0f))
+		.SetPressed(BOX_BRUSH("Common/Button_Pressed", FVector2D(32, 32), 8.0f / 32.0f))
+		.SetNormalPadding(FMargin(2, 2, 2, 2))
+		.SetPressedPadding(FMargin(2, 3, 2, 1));
+	
+	struct ButtonColor
+	{
+	public:
+		FName Name;
+		FLinearColor Normal;
+		FLinearColor Hovered;
+		FLinearColor Pressed;
+
+		ButtonColor(const FName& InName, const FLinearColor& Color) : Name(InName)
+		{
+			Normal = Color * 0.8f;
+			Normal.A = Color.A;
+			Hovered = Color * 1.0f;
+			Hovered.A = Color.A;
+			Pressed = Color * 0.6f;
+			Pressed.A = Color.A;
+		}
+	};
+
+
+	TArray< ButtonColor > FlatButtons;
+	FlatButtons.Add(ButtonColor("SourceFilter.Filter", FLinearColor(0.25f, 0.25f, 0.25f, 0.9f)));
+
+	FlatButtons.Add(ButtonColor("SourceFilter.FilterSetOperation.NOT", FLinearColor(0.85f, 0.25f, 0.25f, 0.9f)));
+	FlatButtons.Add(ButtonColor("SourceFilter.FilterSetOperation.AND", FLinearColor(0.25f, 0.85f, 0.25f, 0.9f)));
+	FlatButtons.Add(ButtonColor("SourceFilter.FilterSetOperation.OR", FLinearColor(0.25f, 0.25f, 0.85f, 0.9f)));
+
+	FlatButtons.Add(ButtonColor("SourceFilter.FilterSet", FLinearColor(0.125f, 0.125f, 0.125f, 0.9f)));
+	
+	for (const ButtonColor& Entry : FlatButtons)
+	{
+		StyleSet->Set(Entry.Name, FButtonStyle(Button)
+			.SetNormal(BOX_BRUSH("Common/FlatButton", 2.0f / 8.0f, Entry.Normal))
+			.SetHovered(BOX_BRUSH("Common/FlatButton", 2.0f / 8.0f, Entry.Hovered))
+			.SetPressed(BOX_BRUSH("Common/FlatButton", 2.0f / 8.0f, Entry.Pressed))
+		);
+	}
+
+	StyleSet->Set("SourceFilter.Splitter", FSplitterStyle()
+		.SetHandleNormalBrush(FSlateColorBrush(FLinearColor(FColor(32, 32, 32))))
+		.SetHandleHighlightBrush(FSlateColorBrush(FLinearColor(FColor(96, 96, 96))))
+	);
+
+	StyleSet->Set("ToggleButton", FButtonStyle(Button)
+		.SetNormal(FSlateNoResource())
+		.SetHovered(BOX_BRUSH("Common/RoundedSelection_16x", 4.0f / 16.0f, FLinearColor(0.701f, 0.225f, 0.003f)))
+		.SetPressed(BOX_BRUSH("Common/RoundedSelection_16x", 4.0f / 16.0f, FLinearColor(0.701f, 0.225f, 0.003f)))
+	);
+
+	FComboButtonStyle ToolbarComboButton = FComboButtonStyle()
+		.SetButtonStyle(StyleSet->GetWidgetStyle<FButtonStyle>("ToggleButton"))
+		.SetDownArrowImage(IMAGE_BRUSH("Common/ShadowComboArrow", Icon8x8))
+		.SetMenuBorderBrush(BOX_BRUSH("Old/Menu_Background", FMargin(8.0f / 64.0f)))
+		.SetMenuBorderPadding(FMargin(0.0f));
+	StyleSet->Set("SourceFilter.ComboButton", ToolbarComboButton);
+
+	StyleSet->Set("SourceFilter.TabIcon", new IMAGE_BRUSH("Icons/AssetIcons/Character_16x", Icon16x16));
+
+	StyleSet->Set("SourceFilter.NewPreset", new IMAGE_BRUSH("Icons/icon_file_saveas_16px", Icon16x16));
+	StyleSet->Set("SourceFilter.LoadPreset", new IMAGE_BRUSH("Icons/icon_file_open_16px", Icon16x16));
+	StyleSet->Set("SourceFilter.SavePreset", new IMAGE_BRUSH("Icons/icon_file_save_16px", Icon16x16));
+
+
+	StyleSet->Set("SourceFilter.TextStyle", FTextBlockStyle(NormalText)
+		.SetFont(DEFAULT_FONT("Bold", 9))
+		.SetColorAndOpacity(FLinearColor(1.0f, 1.0f, 1.0f, 0.9f))
+		.SetShadowOffset(FVector2D(1, 1))
+		.SetShadowColorAndOpacity(FLinearColor(0, 0, 0, 0.9f)));
+
+	StyleSet->Set("SourceFilter.DragDrop.Border", new BOX_BRUSH("Old/Menu_Background", FMargin(8.0f / 64.0f)));
+		
+	FSlateStyleRegistry::RegisterSlateStyle( *StyleSet.Get() );
+}
+
+#undef IMAGE_BRUSH
+#undef BOX_BRUSH
+#undef BORDER_BRUSH
+#undef DEFAULT_FONT
+
+void FSourceFilterStyle::Shutdown()
+{
+	if( StyleSet.IsValid() )
+	{
+		FSlateStyleRegistry::UnRegisterSlateStyle( *StyleSet.Get() );
+		ensure( StyleSet.IsUnique() );
+		StyleSet.Reset();
+	}
+}
+
+const ISlateStyle& FSourceFilterStyle::Get()
+{
+	return *( StyleSet.Get() );
+}
+
+const FName& FSourceFilterStyle::GetStyleSetName()
+{
+	return StyleSet->GetStyleSetName();
+}
diff --git a/Engine/Plugins/Developer/TraceSourceFiltering/Source/SourceFilteringEditor/Private/SourceFilterStyle.h b/Engine/Plugins/Developer/TraceSourceFiltering/Source/SourceFilteringEditor/Private/SourceFilterStyle.h
new file mode 100644
--- /dev/null
+++ b/Engine/Plugins/Developer/TraceSourceFiltering/Source/SourceFilteringEditor/Private/SourceFilterStyle.h
@@ -0,0 +1,29 @@
+// Copyright Epic Games, Inc. All Rights Reserved.
+
+#pragma once
+
+#include "CoreMinimal.h"
+#include "Styling/SlateStyle.h"
+#include "Styling/SlateTypes.h"
+
+class FSourceFilterStyle
+{
+public:
+	static void Initialize();
+	static void Shutdown();
+
+	static const ISlateStyle& Get();
+	static const FName& GetStyleSetName();
+	static const FSlateBrush* GetBrush(FName PropertyName, const ANSICHAR* Specifier = NULL)
+	{
+		return StyleSet->GetBrush(PropertyName, Specifier);
+	}
+
+private:
+	static FSlateBrush* FilterBrush;
+	static FSlateBrush* FilterSetOperationBrush;
+	static FSlateBrush* FilterSetBrush;
+	/** Singleton instances of this style. */
+	static TSharedPtr< class FSlateStyleSet > StyleSet;	
+	static FTextBlockStyle NormalText;
+};
diff --git a/Engine/Plugins/Developer/TraceSourceFiltering/Source/SourceFilteringEditor/Private/SourceFilteringEditorModule.cpp b/Engine/Plugins/Developer/TraceSourceFiltering/Source/SourceFilteringEditor/Private/SourceFilteringEditorModule.cpp
new file mode 100644
--- /dev/null
+++ b/Engine/Plugins/Developer/TraceSourceFiltering/Source/SourceFilteringEditor/Private/SourceFilteringEditorModule.cpp
@@ -0,0 +1,73 @@
+// Copyright Epic Games, Inc. All Rights Reserved.
+
+#include "SourceFilteringEditorModule.h"
+#include "Modules/ModuleManager.h"
+
+#include "TraceServices/SessionService.h"
+#include "TraceServices/ITraceServicesModule.h"
+#include "HAL/IConsoleManager.h"
+#include "Insights/IUnrealInsightsModule.h"
+#include "Framework/Docking/LayoutExtender.h"
+#include "Framework/Docking/TabManager.h"
+#include "Widgets/Docking/SDockTab.h"
+#include "Misc/ConfigCacheIni.h"
+#include "WorkspaceMenuStructureModule.h"
+#include "WorkspaceMenuStructure.h"
+#include "Features/IModularFeatures.h"
+
+#include "STraceSourceFilteringWidget.h"
+#include "SourceFilterStyle.h"
+
+#include "IGameplayInsightsModule.h"
+
+IMPLEMENT_MODULE(FSourceFilteringEditorModule, SourceFilteringEditor);
+
+FString FSourceFilteringEditorModule::SourceFiltersIni;
+FName InsightsSourceFilteringTabName(TEXT("InsightsSourceFiltering"));
+
+void FSourceFilteringEditorModule::StartupModule()
+{
+	FSourceFilterStyle::Initialize();
+
+	// Populate static ini path
+	FConfigCacheIni::LoadGlobalIniFile(SourceFiltersIni, TEXT("TraceSourceFilters"));
+	
+	IUnrealInsightsModule& UnrealInsightsModule = FModuleManager::LoadModuleChecked<IUnrealInsightsModule>("TraceInsights");	
+	FOnRegisterMajorTabExtensions& TimingProfilerLayoutExtension = UnrealInsightsModule.OnRegisterMajorTabExtension(FInsightsManagerTabs::TimingProfilerTabId);
+	TimingProfilerLayoutExtension.AddRaw(this, &FSourceFilteringEditorModule::RegisterLayoutExtensions);
+}
+
+void FSourceFilteringEditorModule::ShutdownModule()
+{
+	FSourceFilterStyle::Shutdown();
+}
+
+void FSourceFilteringEditorModule::RegisterLayoutExtensions(FInsightsMajorTabExtender& InOutExtender)
+{
+#if WITH_EDITOR
+	FTabId ExtendedTabId(GameplayInsightsTabs::DocumentTab);
+#else
+	FTabId ExtendedTabId(FTimingProfilerTabs::TimersID);
+#endif
+
+	InOutExtender.GetLayoutExtender().ExtendLayout(ExtendedTabId, ELayoutExtensionPosition::Before, FTabManager::FTab(InsightsSourceFilteringTabName, ETabState::ClosedTab));
+
+	TSharedRef<FWorkspaceItem> Category = InOutExtender.GetTabManager()->AddLocalWorkspaceMenuCategory(NSLOCTEXT("FInsightsSourceFilteringModule", "SourceFilteringGroupName", "Filtering"));
+
+	FInsightsMinorTabConfig& MinorTabConfig = InOutExtender.AddMinorTabConfig();
+	MinorTabConfig.TabId = InsightsSourceFilteringTabName;
+	MinorTabConfig.TabLabel = NSLOCTEXT("SourceFilteringEditorModule", "SourceFilteringTab", "Trace Source Filtering");
+	MinorTabConfig.TabTooltip = NSLOCTEXT("SourceFilteringEditorModule", "SourceFilteringTabTooltip", "Opens the Trace Source Filtering tab, allows for filtering UWorld and AActor instances to not output Trace data");
+	MinorTabConfig.TabIcon = FSlateIcon(FSourceFilterStyle::GetStyleSetName(), "SourceFilter.TabIcon");
+	MinorTabConfig.WorkspaceGroup = Category;
+	MinorTabConfig.OnSpawnTab = FOnSpawnTab::CreateLambda([](const FSpawnTabArgs& Args)
+	{
+		const TSharedRef<SDockTab> DockTab = SNew(SDockTab)
+			.TabRole(ETabRole::PanelTab);
+
+		TSharedRef<STraceSourceFilteringWidget> Window = SNew(STraceSourceFilteringWidget);
+		DockTab->SetContent(Window);
+
+		return DockTab;
+	});
+}
diff --git a/Engine/Plugins/Developer/TraceSourceFiltering/Source/SourceFilteringEditor/Private/SourceFilteringEditorModule.h b/Engine/Plugins/Developer/TraceSourceFiltering/Source/SourceFilteringEditor/Private/SourceFilteringEditorModule.h
new file mode 100644
--- /dev/null
+++ b/Engine/Plugins/Developer/TraceSourceFiltering/Source/SourceFilteringEditor/Private/SourceFilteringEditorModule.h
@@ -0,0 +1,22 @@
+// Copyright Epic Games, Inc. All Rights Reserved.
+
+#pragma once
+
+#include "CoreMinimal.h"
+#include "Modules/ModuleInterface.h"
+
+#include "Templates/SharedPointer.h"
+
+struct FInsightsMajorTabExtender;
+
+class FSourceFilteringEditorModule : public IModuleInterface
+{
+public:
+	virtual void StartupModule() override;
+	virtual void ShutdownModule() override;
+
+	static FString SourceFiltersIni;
+
+protected:
+	void RegisterLayoutExtensions(FInsightsMajorTabExtender& InOutExtender);
+};
diff --git a/Engine/Plugins/Developer/TraceSourceFiltering/Source/SourceFilteringEditor/Private/TreeViewBuilder.cpp b/Engine/Plugins/Developer/TraceSourceFiltering/Source/SourceFilteringEditor/Private/TreeViewBuilder.cpp
new file mode 100644
--- /dev/null
+++ b/Engine/Plugins/Developer/TraceSourceFiltering/Source/SourceFilteringEditor/Private/TreeViewBuilder.cpp
@@ -0,0 +1,20 @@
+// Copyright Epic Games, Inc. All Rights Reserved.
+
+#include "TreeViewBuilder.h"
+#include "IFilterObject.h"
+
+void FTreeViewDataBuilder::AddFilterObject(TSharedRef<IFilterObject> FilterObject)
+{	
+	FilterObjects.Add(FilterObject);
+	FlatFilterObjects.Add(FilterObject);
+}
+
+void FTreeViewDataBuilder::AddChildFilterObject(TArray<TSharedPtr<IFilterObject>> InFilterObjects, TSharedRef<IFilterObject> ParentObject)
+{
+	for (TSharedPtr<IFilterObject> ChildFilter: InFilterObjects)
+	{
+		FlatFilterObjects.Add(ChildFilter);
+	}
+
+	ParentToChildren.FindOrAdd(ParentObject).Append(InFilterObjects);
+}
diff --git a/Engine/Plugins/Developer/TraceSourceFiltering/Source/SourceFilteringEditor/Private/TreeViewBuilder.h b/Engine/Plugins/Developer/TraceSourceFiltering/Source/SourceFilteringEditor/Private/TreeViewBuilder.h
new file mode 100644
--- /dev/null
+++ b/Engine/Plugins/Developer/TraceSourceFiltering/Source/SourceFilteringEditor/Private/TreeViewBuilder.h
@@ -0,0 +1,28 @@
+// Copyright Epic Games, Inc. All Rights Reserved.
+
+#pragma once
+
+#include "Containers/Array.h"
+#include "Containers/Map.h"
+#include "Templates/SharedPointer.h"
+
+class IFilterObject;
+
+/** Helper structure allowing a ISessionSourceFilterService to populate the data use to generate the SSourceFilteringTreeView in STraceSourceFilteringWidget */
+struct FTreeViewDataBuilder
+{
+	FTreeViewDataBuilder(TArray<TSharedPtr<IFilterObject>>& InFilterObjects, TMap<TSharedPtr<IFilterObject>, TArray<TSharedPtr<IFilterObject>>>& InParentToChildren, TArray<TSharedPtr<IFilterObject>>& InFlatFilterObjects) : FilterObjects(InFilterObjects), ParentToChildren(InParentToChildren), FlatFilterObjects(InFlatFilterObjects) {}
+
+	/** Add a root-level filter */
+	void AddFilterObject(TSharedRef<IFilterObject> FilterObject);
+	/** Add a set of child-filters for a, previously added, parent filter */
+	void AddChildFilterObject(TArray<TSharedPtr<IFilterObject>> FilterObjects, TSharedRef<IFilterObject> ParentObject);
+
+protected:
+	/** Root-level filter objects  */
+	TArray<TSharedPtr<IFilterObject>>& FilterObjects;
+	/** Parent / child filter(s) mapping */
+	TMap<TSharedPtr<IFilterObject>, TArray<TSharedPtr<IFilterObject>>>& ParentToChildren;
+	/** Flat list of all the filter objects contained in both FilterObjects and ParentToChildren */
+	TArray<TSharedPtr<IFilterObject>>& FlatFilterObjects;
+};
\ No newline at end of file
diff --git a/Engine/Plugins/Developer/TraceSourceFiltering/Source/SourceFilteringEditor/Private/WorldFilters.cpp b/Engine/Plugins/Developer/TraceSourceFiltering/Source/SourceFilteringEditor/Private/WorldFilters.cpp
new file mode 100644
--- /dev/null
+++ b/Engine/Plugins/Developer/TraceSourceFiltering/Source/SourceFilteringEditor/Private/WorldFilters.cpp
@@ -0,0 +1,248 @@
+// Copyright Epic Games, Inc. All Rights Reserved.
+
+#include "WorldFilters.h"
+#include "CoreGlobals.h"
+#include "Misc/ConfigCacheIni.h"
+#include "Algo/Transform.h"
+
+#include "SourceFilteringEditorModule.h"
+
+#if WITH_ENGINE
+#include "Engine/EngineBaseTypes.h"
+#include "Engine/EngineTypes.h"
+#endif // WITH_ENGINE
+
+#define LOCTEXT_NAMESPACE "FWorldTypeTraceFilter"
+
+FWorldTypeTraceFilter::FWorldTypeTraceFilter(const TFunction<void(uint8, bool)>& InOnSetWorldTypeFilterState, const TFunction<bool(uint8)>& InOnGetWorldTypeFilterState) : OnSetWorldTypeFilterState(InOnSetWorldTypeFilterState), OnGetWorldTypeFilterState(InOnGetWorldTypeFilterState)
+{
+	WorldTypeFilterName = LOCTEXT("WorldTypeFilterName", "By World Type");
+	WorldTypeFilterValues.Add(TPair<FText, TArray<uint8>>(LOCTEXT("EditorWorldFilterName", "Editor"), TArray<uint8>({ 2 /*EWorldType::Editor*/, 4 /*EWorldType::EditorPreview*/ })));
+	WorldTypeFilterValues.Add(TPair<FText, TArray<uint8>>(LOCTEXT("PIEWorldFilterName", "PIE"), TArray<uint8>({ 3 /*EWorldType::PIE*/ })));
+	WorldTypeFilterValues.Add(TPair<FText, TArray<uint8>>(LOCTEXT("GameWorldFilterName", "Game"), TArray<uint8>({ 1 /*EWorldType::Game*/, 5 /*EWorldType::GamePreview*/, 6 /*EWorldType::GameRPC*/ })));
+
+#if WITH_ENGINE
+	static_assert((uint8)EWorldType::Editor == 2, "EWorldType and WorldTypeFilterValues must be kept in sync");
+	static_assert((uint8)EWorldType::EditorPreview == 4, "EWorldType and WorldTypeFilterValues must be kept in sync");
+	static_assert((uint8)EWorldType::PIE == 3, "EWorldType and WorldTypeFilterValues must be kept in sync");
+	static_assert((uint8)EWorldType::Game == 1, "EWorldType and WorldTypeFilterValues must be kept in sync");
+	static_assert((uint8)EWorldType::GamePreview == 5, "EWorldType and WorldTypeFilterValues must be kept in sync");
+	static_assert((uint8)EWorldType::GameRPC == 6, "EWorldType and WorldTypeFilterValues must be kept in sync");
+#endif
+
+	LoadSettings();
+}
+
+FText FWorldTypeTraceFilter::GetDisplayText()
+{
+	return LOCTEXT("WorldTypeTraceFilterLabel", "Filter by World Type");
+}
+
+FText FWorldTypeTraceFilter::GetToolTipText()
+{
+	return LOCTEXT("WorldTypeTraceFilterToolTip", "Marks World's on the analyis session as Traceable if it is of (any) specified World Type(s)");
+}
+
+void FWorldTypeTraceFilter::PopulateMenuBuilder(FMenuBuilder& InBuilder)
+{
+	for (const TPair<FText, TArray<uint8>>& TypePair : WorldTypeFilterValues)
+	{
+		InBuilder.AddMenuEntry(TypePair.Key,
+			FText::Format(LOCTEXT("WorldTypeFilterTooltipFormat", "Set Filtered state for World Type: {0}"), TypePair.Key),
+			FSlateIcon(),
+			FUIAction(
+				FExecuteAction::CreateLambda([this, TypePair]()
+				{
+					const bool bNewState = !OnGetWorldTypeFilterState(TypePair.Value[0]);
+					for (const uint8& Type : TypePair.Value)
+					{
+						OnSetWorldTypeFilterState(Type, bNewState);
+					}
+
+					SaveSettings();
+				}),
+				FCanExecuteAction(),
+				FGetActionCheckState::CreateLambda([this, TypePair]()
+				{
+					return OnGetWorldTypeFilterState(TypePair.Value[0]) ? ECheckBoxState::Checked : ECheckBoxState::Unchecked;
+				})
+			),
+			NAME_None,
+			EUserInterfaceActionType::ToggleButton
+		);
+	}
+}
+
+void FWorldTypeTraceFilter::SaveSettings()
+{
+	FString IniSectionName = TEXT("WorldTypeTraceFilter");
+	FString IniDisabledTypesKey = TEXT("DisabledWorldTypes");
+
+	FString Separator = TEXT("+");
+	FString DisabledWorldTypeString;
+	for (const TPair<FText, TArray<uint8>>& WorldTypePair : WorldTypeFilterValues)
+	{
+		for (uint8 WorldTypeValue : WorldTypePair.Value)
+		{
+			const bool bDisabled = !OnGetWorldTypeFilterState(WorldTypeValue);
+			if (bDisabled)
+			{
+				DisabledWorldTypeString += FString::FromInt(WorldTypeValue);
+				DisabledWorldTypeString += Separator;
+			}
+		}
+	}
+	DisabledWorldTypeString.RemoveFromEnd(Separator);	
+
+	GConfig->SetString(*IniSectionName, *IniDisabledTypesKey, *DisabledWorldTypeString, FSourceFilteringEditorModule::SourceFiltersIni);
+}
+
+void FWorldTypeTraceFilter::LoadSettings()
+{
+	FString IniSectionName = TEXT("WorldTypeTraceFilter");
+	FString IniDisabledTypesKey = TEXT("DisabledWorldTypes");
+
+	FString DisabledTypesString;
+	GConfig->GetString(*IniSectionName, *IniDisabledTypesKey, DisabledTypesString, FSourceFilteringEditorModule::SourceFiltersIni);
+
+	FString Separator = TEXT("+");
+	TArray<FString> DisabledWorldTypeStrings;
+	DisabledTypesString.ParseIntoArray(DisabledWorldTypeStrings, *Separator);
+	
+	TArray<uint8> DisabledWorldTypes;
+	Algo::Transform(DisabledWorldTypeStrings, DisabledWorldTypes, [](FString Value)
+	{
+		return FCString::Atoi(*Value);
+	});
+
+	for (const TPair<FText, TArray<uint8>>& WorldTypePair : WorldTypeFilterValues)
+	{
+		for (uint8 WorldTypeValue : WorldTypePair.Value)
+		{
+			const bool bState = OnGetWorldTypeFilterState(WorldTypeValue);
+			const bool bIniState = !DisabledWorldTypes.Contains(WorldTypeValue);
+
+			if (bState != bIniState)
+			{
+				OnSetWorldTypeFilterState(WorldTypeValue, bIniState);
+			}
+		}
+	}
+}
+
+FWorldNetModeTraceFilter::FWorldNetModeTraceFilter(const TFunction<void(uint8, bool)>& InOnSetWorldNetModeFilterState, const TFunction<bool(uint8)>& InOnGetWorldNetModeFilterState) : OnSetWorldNetModeFilterState(InOnSetWorldNetModeFilterState), OnGetWorldNetModeFilterState(InOnGetWorldNetModeFilterState)
+{
+	WorldNetModeFilterName = LOCTEXT("ByNetModeWorldFilterName", "By Net Mode");
+
+	WorldNetModeFilterValues.Add(TPair<FText, TArray<uint8>>(LOCTEXT("ClientNetModeWorldFilterName", "Client"), TArray<uint8>({ 3 /*NM_Client*/ })));
+	WorldNetModeFilterValues.Add(TPair<FText, TArray<uint8>>(LOCTEXT("ServerNetModeWorldFilterName", "Server"), TArray<uint8>({ 1 /*NM_DedicatedServer*/, 2 /*NM_ListenServer*/ })));
+	WorldNetModeFilterValues.Add(TPair<FText, TArray<uint8>>(LOCTEXT("StandaloneNetModeWorldFilterName", "Standalone"), TArray<uint8>({ 0 /*NM_Standalone*/ })));
+	
+#if WITH_ENGINE
+	static_assert((uint8)ENetMode::NM_Client == 3, "ENetMode and WorldNetModeFilterValues must be kept in sync");
+	static_assert((uint8)ENetMode::NM_DedicatedServer == 1, "ENetMode and WorldNetModeFilterValues must be kept in sync");
+	static_assert((uint8)ENetMode::NM_ListenServer == 2, "ENetMode and WorldNetModeFilterValues must be kept in sync");
+	static_assert((uint8)ENetMode::NM_Standalone == 0, "ENetMode and WorldNetModeFilterValues must be kept in sync");
+#endif
+
+	LoadSettings();
+}
+
+FText FWorldNetModeTraceFilter::GetDisplayText()
+{
+	return LOCTEXT("WorldNetModeTraceFilterLabel", "Filter by World Net Mode");
+}
+
+FText FWorldNetModeTraceFilter::GetToolTipText()
+{
+	return LOCTEXT("WorldNetModeTraceFilterToolTip", "Marks World's on the analysis session as Traceable if it is running in (any) of the specified Net Mode(s)");
+}
+
+void FWorldNetModeTraceFilter::PopulateMenuBuilder(FMenuBuilder& InBuilder)
+{
+	for (const TPair<FText, TArray<uint8>>& TypePair : WorldNetModeFilterValues)
+	{
+		InBuilder.AddMenuEntry(TypePair.Key,
+			FText::Format(LOCTEXT("NetModeFilterTooltipFormat", "Set Filtered state for Net Mode: {0}"), TypePair.Key),
+			FSlateIcon(),
+			FUIAction(
+				FExecuteAction::CreateLambda([this, TypePair]()
+				{
+					const bool bNewState = !OnGetWorldNetModeFilterState(TypePair.Value[0]);
+					for (const uint8& Type : TypePair.Value)
+					{
+						OnSetWorldNetModeFilterState(Type, bNewState);
+					}
+
+					SaveSettings();
+				}),
+				FCanExecuteAction(),
+				FGetActionCheckState::CreateLambda([this, TypePair]()
+				{
+					return OnGetWorldNetModeFilterState(TypePair.Value[0]) ? ECheckBoxState::Checked : ECheckBoxState::Unchecked;
+				})
+			),
+			NAME_None,
+			EUserInterfaceActionType::ToggleButton
+		);
+	}
+}
+
+void FWorldNetModeTraceFilter::SaveSettings()
+{
+	FString IniSectionName = TEXT("WorldNetModeTraceFilter");
+	FString IniDisabledTypesKey = TEXT("DisabledNetModeTypes");
+
+	FString Separator = TEXT("+");
+	FString DisabledWorldNetModeString;
+	for (const TPair<FText, TArray<uint8>>& WorldNetModePair : WorldNetModeFilterValues)
+	{
+		for (uint8 WorldNetModeValue : WorldNetModePair.Value)
+		{
+			const bool bDisabled = !OnGetWorldNetModeFilterState(WorldNetModeValue);
+			if (bDisabled)
+			{
+				DisabledWorldNetModeString += FString::FromInt(WorldNetModeValue);
+				DisabledWorldNetModeString += Separator;
+			}
+		}
+	}
+	DisabledWorldNetModeString.RemoveFromEnd(Separator);
+
+	GConfig->SetString(*IniSectionName, *IniDisabledTypesKey, *DisabledWorldNetModeString, FSourceFilteringEditorModule::SourceFiltersIni);
+}
+
+void FWorldNetModeTraceFilter::LoadSettings()
+{
+	FString IniSectionName = TEXT("WorldNetModeTraceFilter");
+	FString IniDisabledTypesKey = TEXT("DisabledNetModeTypes");
+
+	FString DisabledTypesString;
+	GConfig->GetString(*IniSectionName, *IniDisabledTypesKey, DisabledTypesString, FSourceFilteringEditorModule::SourceFiltersIni);
+
+	FString Separator = TEXT("+");
+	TArray<FString> DisabledWorldNetModeStrings;
+	DisabledTypesString.ParseIntoArray(DisabledWorldNetModeStrings, *Separator);
+
+	TArray<uint8> DisabledWorldNetModes;
+	Algo::Transform(DisabledWorldNetModeStrings, DisabledWorldNetModes, [](FString Value)
+	{
+		return FCString::Atoi(*Value);
+	});
+
+	for (const TPair<FText, TArray<uint8>>& WorldNetModePair : WorldNetModeFilterValues)
+	{
+		for (uint8 WorldNetModeValue : WorldNetModePair.Value)
+		{
+			const bool bState = OnGetWorldNetModeFilterState(WorldNetModeValue);
+			const bool bIniState = !DisabledWorldNetModes.Contains(WorldNetModeValue);
+
+			if (bState != bIniState)
+			{
+				OnSetWorldNetModeFilterState(WorldNetModeValue, bIniState);
+			}
+		}
+	}
+}
+
+#undef LOCTEXT_NAMESPACE // "FWorldTypeTraceFilter"
\ No newline at end of file
diff --git a/Engine/Plugins/Developer/TraceSourceFiltering/Source/SourceFilteringEditor/Private/WorldFilters.h b/Engine/Plugins/Developer/TraceSourceFiltering/Source/SourceFilteringEditor/Private/WorldFilters.h
new file mode 100644
--- /dev/null
+++ b/Engine/Plugins/Developer/TraceSourceFiltering/Source/SourceFilteringEditor/Private/WorldFilters.h
@@ -0,0 +1,56 @@
+// Copyright Epic Games, Inc. All Rights Reserved.
+
+#pragma once
+
+#include "Framework/MultiBox/MultiBoxBuilder.h"
+#include "Templates/Function.h"
+
+#include "ISessionSourceFilterService.h"
+#include "Internationalization/Internationalization.h"
+
+/** World filter implementation, that marks Worlds as traceable according to their EWorldType */
+class FWorldTypeTraceFilter : public IWorldTraceFilter
+{
+public:
+	FWorldTypeTraceFilter(const TFunction<void(uint8, bool)>& InOnSetWorldTypeFilterState, const TFunction<bool(uint8)>& InOnGetWorldTypeFilterState);
+
+	/** Begin IWorldTraceFilter Overrides */
+	virtual FText GetDisplayText() override;
+	virtual FText GetToolTipText() override;
+	virtual void PopulateMenuBuilder(FMenuBuilder& InBuilder) override;
+	/** End IWorldTraceFilter Overrides */
+protected:	
+	/** Load / save any settings to config that should be persistent between editor sessions */
+	void SaveSettings();	
+	void LoadSettings();
+
+protected:
+	TFunction<void(uint8, bool)> OnSetWorldTypeFilterState;
+	TFunction<bool(uint8)> OnGetWorldTypeFilterState;
+
+	FText WorldTypeFilterName;
+	TArray<TPair<FText, TArray<uint8>>> WorldTypeFilterValues;
+};
+
+/** World filter implementation, that marks Worlds as traceable according to their ENetMode */
+class FWorldNetModeTraceFilter : public IWorldTraceFilter
+{
+public:
+	FWorldNetModeTraceFilter(const TFunction<void(uint8, bool)>& InOnSetWorldNetModeFilterState, const TFunction<bool(uint8)>& InOnGetWorldNetModeFilterState);
+	
+	/** Begin IWorldTraceFilter Overrides */
+	virtual FText GetDisplayText() override;
+	virtual FText GetToolTipText() override;
+	virtual void PopulateMenuBuilder(FMenuBuilder& InBuilder) override;
+	/** End IWorldTraceFilter Overrides */
+protected:
+	/** Load / save any settings to config that should be persistent between editor sessions */
+	void SaveSettings();
+	void LoadSettings();
+protected:
+	TFunction<void(uint8, bool)> OnSetWorldNetModeFilterState;
+	TFunction<bool(uint8)> OnGetWorldNetModeFilterState;
+
+	FText WorldNetModeFilterName;
+	TArray<TPair<FText, TArray<uint8>>> WorldNetModeFilterValues;
+};
diff --git a/Engine/Plugins/Developer/TraceSourceFiltering/Source/SourceFilteringEditor/Private/WorldObject.h b/Engine/Plugins/Developer/TraceSourceFiltering/Source/SourceFilteringEditor/Private/WorldObject.h
new file mode 100644
--- /dev/null
+++ b/Engine/Plugins/Developer/TraceSourceFiltering/Source/SourceFilteringEditor/Private/WorldObject.h
@@ -0,0 +1,24 @@
+// Copyright Epic Games, Inc. All Rights Reserved.
+
+#pragma once
+
+#include "Containers/UnrealString.h"
+#include "Templates/SharedPointer.h"
+
+/** Object representing an UWorld and its filtering state */
+class FWorldObject : public TSharedFromThis<FWorldObject>
+{
+public:
+	FWorldObject(const FString& InName, uint8 InWorldType, bool bState, uint32 InHash) : Name(InName), WorldType(InWorldType), bCanOutputData(bState), Hash(InHash) {}
+
+	const FText GetDisplayText() const { return FText::FromString(Name); }
+	const FString& GetName() const { return Name; }
+	bool CanOutputData() const { return bCanOutputData; }
+	const uint32 GetHash() const { return Hash; }
+protected:
+	FString Name;
+	// Corresponds to EWorldType in Engine
+	uint8 WorldType;
+	bool bCanOutputData;
+	uint32 Hash;
+};
\ No newline at end of file
diff --git a/Engine/Plugins/Developer/TraceSourceFiltering/Source/SourceFilteringEditor/SourceFilteringEditor.Build.cs b/Engine/Plugins/Developer/TraceSourceFiltering/Source/SourceFilteringEditor/SourceFilteringEditor.Build.cs
new file mode 100644
--- /dev/null
+++ b/Engine/Plugins/Developer/TraceSourceFiltering/Source/SourceFilteringEditor/SourceFilteringEditor.Build.cs
@@ -0,0 +1,48 @@
+// Copyright Epic Games, Inc. All Rights Reserved.
+
+using UnrealBuildTool;
+
+public class SourceFilteringEditor : ModuleRules
+{
+    public SourceFilteringEditor(ReadOnlyTargetRules Target) : base(Target)
+    {
+        PrivateDependencyModuleNames.AddRange(
+                new string[] {
+                    "Core",
+                    "CoreUObject",
+                    "Slate",
+                    "SlateCore",
+                    "EditorStyle",
+                    "InputCore",
+                    "TraceServices",
+                    "TraceInsights",
+					"TraceAnalysis",
+					"WorkspaceMenuStructure",
+                    "SourceFilteringCore",
+					"PropertyPath",
+					"GameplayInsights",
+					"Sockets",
+                    "MessagingCommon"
+                }
+            );
+		
+        if (Target.bCompileAgainstEngine)
+        {
+            PrivateDependencyModuleNames.AddRange(
+				new string[] {
+					 "Engine",
+					 "SourceFilteringTrace"
+				}
+			);
+        }
+
+        if (Target.bBuildEditor)
+        {
+            PrivateDependencyModuleNames.AddRange(
+				new string[] {
+					"UnrealEd",                 
+				}
+			);
+        }
+    }
+}
diff --git a/Engine/Plugins/Developer/TraceSourceFiltering/Source/SourceFilteringTrace/Private/DataSourceFilterSet.cpp b/Engine/Plugins/Developer/TraceSourceFiltering/Source/SourceFilteringTrace/Private/DataSourceFilterSet.cpp
new file mode 100644
--- /dev/null
+++ b/Engine/Plugins/Developer/TraceSourceFiltering/Source/SourceFilteringTrace/Private/DataSourceFilterSet.cpp
@@ -0,0 +1,137 @@
+// Copyright Epic Games, Inc. All Rights Reserved.
+
+#include "DataSourceFilterSet.h"
+#include "Algo/Transform.h"
+#include "SourceFilterTrace.h"
+
+void UDataSourceFilterSet::SetFilterMode(EFilterSetMode InMode)
+{
+	TRACE_FILTER_OPERATION(this, ESourceActorFilterOperation::SetFilterMode, (uint64)InMode);
+	Mode = InMode;
+}
+
+bool UDataSourceFilterSet::DoesActorPassFilter_Internal(const AActor* InActor) const
+{
+	bool bPassesFilter = [this]()
+	{
+		if (Filters.Num() == 0)
+		{
+			return true;
+		}
+
+		switch (Mode)
+		{
+		case EFilterSetMode::OR: 
+			return false;
+		case EFilterSetMode::AND: 
+			return true;
+		case EFilterSetMode::NOT: 
+			return true;
+		default: 
+			return false;
+		}
+	}();		
+		
+	for (const UDataSourceFilter* Filter : Filters)
+	{
+		if (Filter->IsEnabled())
+		{
+			const bool bResult = Filter->DoesActorPassFilter(InActor);
+
+			if (!bResult)
+			{
+				FFilterLedger::Get().RejectedFilters.Add(Filter);
+			}
+
+			// OR, at least one filter has to be passed 
+			if (Mode == EFilterSetMode::OR)
+			{
+				bPassesFilter |= bResult;
+				if (bPassesFilter)
+				{
+					break;
+				}
+			}
+			// AND, all filters have to be passed 
+			else if (Mode == EFilterSetMode::AND)
+			{
+				bPassesFilter &= bResult;
+
+				if (!bPassesFilter)
+				{
+					break;
+				}
+			}
+			// NOT, all filters have to be !passed
+			else if (Mode == EFilterSetMode::NOT)
+			{
+				bPassesFilter &= !bResult;
+
+				if (!bPassesFilter)
+				{
+					break;
+				}
+			}
+		}
+	}
+
+	return bPassesFilter;
+}
+
+void UDataSourceFilterSet::GetDisplayText_Internal(FText& OutDisplayText) const
+{
+	TArray<FString> PerFilterTextStrings;
+	for (const UDataSourceFilter* Filter : Filters)
+	{
+		FText FilterText;
+		Execute_GetDisplayText(Filter, FilterText);
+
+		PerFilterTextStrings.Add(FilterText.ToString());
+	}
+
+	const FString Delimiter = [this]()
+	{
+		if (Mode == EFilterSetMode::OR)
+		{
+			return TEXT(" || ");
+		}
+		else if (Mode == EFilterSetMode::AND)
+		{
+			return TEXT(" && ");
+		}
+		else if (Mode == EFilterSetMode::NOT)
+		{
+			return TEXT(" !");
+		}
+
+		return TEXT(" ");
+	}();
+	
+	FString ComposedString = TEXT("(");
+	for (const FString& FilterString : PerFilterTextStrings)
+	{
+		ComposedString += FilterString;
+		ComposedString += Delimiter;
+	}
+
+	ComposedString.RemoveFromEnd(Delimiter);
+	ComposedString += TEXT(")");
+
+	OutDisplayText = FText::FromString(ComposedString);
+}
+
+EFilterSetMode UDataSourceFilterSet::GetFilterSetMode() const
+{
+	return Mode;
+}
+
+void UDataSourceFilterSet::SetEnabled(bool bState)
+{
+	Super::SetEnabled(bState);
+
+	for (UDataSourceFilter* ChildFilter : Filters)
+	{
+		ChildFilter->SetEnabled(bState);
+	}
+}
+
diff --git a/Engine/Plugins/Developer/TraceSourceFiltering/Source/SourceFilteringTrace/Private/SourceFilterCollection.cpp b/Engine/Plugins/Developer/TraceSourceFiltering/Source/SourceFilteringTrace/Private/SourceFilterCollection.cpp
new file mode 100644
--- /dev/null
+++ b/Engine/Plugins/Developer/TraceSourceFiltering/Source/SourceFilteringTrace/Private/SourceFilterCollection.cpp
@@ -0,0 +1,395 @@
+// Copyright Epic Games, Inc. All Rights Reserved.
+
+#include "SourceFilterCollection.h"
+#include "SourceFilterTrace.h"
+#include "AssetRegistryModule.h"
+#include "AssetData.h"
+#include "Async/Async.h"
+#include "UObject/UObjectIterator.h"
+#include "EmptySourceFilter.h"
+
+template<typename T>
+T* USourceFilterCollection::CreateNewFilter(UClass* Class /*= T::StaticClass()*/)
+{
+	T* NewFilter = NewObject<T>(this, Class, NAME_None, RF_Transactional);
+	return NewFilter;
+}
+
+
+void USourceFilterCollection::AddFilter(UDataSourceFilter* NewFilter)
+{
+	Filters.Add(NewFilter);
+	ChildToParent.Add(NewFilter, nullptr);
+	AddClassName(NewFilter);
+
+	TRACE_FILTER_INSTANCE(NewFilter);
+}
+
+UDataSourceFilter* USourceFilterCollection::AddFilterOfClass(const TSubclassOf<UDataSourceFilter>& FilterClass)
+{
+	checkf(FilterClass.Get(), TEXT("Cannot create filter using a null class"));
+
+	UDataSourceFilter* NewFilter = CreateNewFilter<UDataSourceFilter>(FilterClass.Get());
+	AddFilter(NewFilter);
+	AddClassName(NewFilter);
+
+	return NewFilter;
+}
+
+UDataSourceFilter* USourceFilterCollection::AddFilterOfClassToSet(const TSubclassOf<UDataSourceFilter>& FilterClass, UDataSourceFilterSet* FilterSet)
+{
+	ensure(FilterClass.Get() && FilterSet);
+	FilterSet->Modify();
+
+	UDataSourceFilter* NewFilter = CreateNewFilter<UDataSourceFilter>(FilterClass.Get());
+	FilterSet->Filters.Add(NewFilter);
+	ChildToParent.Add(NewFilter, FilterSet);
+	AddClassName(NewFilter);
+
+	TRACE_FILTER_INSTANCE(NewFilter);
+
+	TRACE_FILTER_OPERATION(NewFilter, ESourceActorFilterOperation::MoveFilter, FObjectTrace::GetObjectId(FilterSet));
+
+	return NewFilter;
+}
+
+UDataSourceFilterSet* USourceFilterCollection::ConvertFilterToSet(UDataSourceFilter* ReplacedFilter, EFilterSetMode Mode)
+{
+	ensure(ReplacedFilter);
+
+	UDataSourceFilterSet* NewFilterSet = CreateNewFilter<UDataSourceFilterSet>();
+	TRACE_FILTER_SET(NewFilterSet);
+	AddClassName(NewFilterSet);
+
+	NewFilterSet->SetFilterMode(Mode);
+
+	ReplaceFilter(ReplacedFilter, NewFilterSet);
+	MoveFilter(ReplacedFilter, NewFilterSet);
+
+	return NewFilterSet;
+}
+
+UDataSourceFilterSet* USourceFilterCollection::MakeFilterSet(UDataSourceFilter* FilterOne, UDataSourceFilter* FilterTwo, EFilterSetMode Mode)
+{
+	ensure(FilterOne && FilterTwo);
+
+	UDataSourceFilterSet* NewFilterSet = CreateNewFilter<UDataSourceFilterSet>();
+	TRACE_FILTER_SET(NewFilterSet);
+	AddClassName(NewFilterSet);
+
+	NewFilterSet->SetFilterMode(Mode);
+
+	ReplaceFilter(FilterOne, NewFilterSet);
+
+	MoveFilter(FilterOne, NewFilterSet);
+	MoveFilter(FilterTwo, NewFilterSet);
+
+	return NewFilterSet;
+}
+
+
+UDataSourceFilterSet* USourceFilterCollection::MakeEmptyFilterSet(EFilterSetMode Mode)
+{
+	UDataSourceFilterSet* NewFilterSet = CreateNewFilter<UDataSourceFilterSet>();
+	TRACE_FILTER_SET(NewFilterSet);
+	AddClassName(NewFilterSet);
+
+	NewFilterSet->SetFilterMode(Mode);
+
+	Filters.Add(NewFilterSet);
+	ChildToParent.FindOrAdd(NewFilterSet) = nullptr;
+
+	return NewFilterSet;
+}
+
+void USourceFilterCollection::AddFiltersFromPreset(const TArray<FString>& ClassNames, const TMap<int32, int32>& ChildToParentIndices)
+{
+	AsyncTask(ENamedThreads::GameThread,
+		[this, ClassNames, ChildToParentIndices]()
+	{
+		const UEnum* EnumPtr = FindObject<UEnum>(ANY_PACKAGE, TEXT("EFilterSetMode"), true);
+
+		TArray<UDataSourceFilter*> CreatedFilters;
+		TArray<bool> MovedFlags;
+
+		CreatedFilters.SetNumZeroed(ClassNames.Num());
+		MovedFlags.SetNumZeroed(ClassNames.Num());
+
+		for (const TPair<int32, int32>& ChildToParentPair : ChildToParentIndices)
+		{
+			const int32 FilterIndex = ChildToParentPair.Key;
+			const int64 Value = EnumPtr->GetValueByNameString(ClassNames[FilterIndex]);
+			if (Value != INDEX_NONE)
+			{
+				// This is a set, so create one
+				CreatedFilters[FilterIndex] = MakeEmptyFilterSet((EFilterSetMode)Value);
+			}
+			else
+			{
+				// This is a filter, so create it according to the class name				
+				if (UClass* Class = FSourceFilterTrace::RetrieveClassByName(ClassNames[FilterIndex]))
+				{
+					if (UDataSourceFilter* Filter = CreateNewFilter<UDataSourceFilter>(Class))
+					{
+						CreatedFilters[FilterIndex] = Filter;
+						AddFilter(Filter);
+					}
+				}
+			}
+
+			// In case we were unable to retrieve a Filter instance its class replace it with an empty 'stub' filter instance, indicating to the user that it, and its class, is missing
+			if (CreatedFilters[FilterIndex] == nullptr)
+			{
+				if (UEmptySourceFilter* Filter = CreateNewFilter<UEmptySourceFilter>())
+				{
+					TRACE_FILTER_CLASS(UEmptySourceFilter::StaticClass());
+					CreatedFilters[FilterIndex] = Filter;
+
+					Filter->MissingClassName = ClassNames[FilterIndex];
+					AddFilter(Filter);
+				}
+
+			}
+
+			if (CreatedFilters[FilterIndex] != nullptr && ChildToParentPair.Value == INDEX_NONE)
+			{
+				MovedFlags[FilterIndex] = true;
+			}
+		}
+
+		// Regenerate the parent child relationship from the preset
+		bool bAnyMoved = true;
+		while (bAnyMoved)
+		{
+			bAnyMoved = false;
+			for (const TPair<int32, int32>& ChildToParentPair : ChildToParentIndices)
+			{
+				if (ChildToParentPair.Value != INDEX_NONE && !MovedFlags[ChildToParentPair.Key] && MovedFlags[ChildToParentPair.Value])
+				{
+					MoveFilter(CreatedFilters[ChildToParentPair.Key], CastChecked<UDataSourceFilterSet>(CreatedFilters[ChildToParentPair.Value]));
+					MovedFlags[ChildToParentPair.Key] = true;
+					bAnyMoved = true;
+				}
+			}
+		}
+	});
+}
+
+void USourceFilterCollection::RemoveFilter(UDataSourceFilter* ToRemoveFilter)
+{
+	UDataSourceFilterSet* OuterFilterSet = nullptr;
+	ChildToParent.RemoveAndCopyValue(ToRemoveFilter, OuterFilterSet);
+	
+	// In case of a set, also remove children 
+	if (UDataSourceFilterSet* FilterSet = Cast<UDataSourceFilterSet>(ToRemoveFilter))
+	{
+		FilterSet->Modify();
+
+		// Also remove contained children
+		for (UDataSourceFilter* Filter : FilterSet->Filters)
+		{
+			RemoveFilterRecursive(Filter);
+		}
+
+		FilterSet->Filters.Empty();
+	}
+
+	TRACE_FILTER_OPERATION(ToRemoveFilter, ESourceActorFilterOperation::RemoveFilter, FObjectTrace::GetObjectId(Cast<UDataSourceFilterSet>(ToRemoveFilter)));
+
+	if (OuterFilterSet)
+	{
+		OuterFilterSet->Modify();
+		RemoveFilterFromSet(ToRemoveFilter, OuterFilterSet);
+	}
+	else
+	{
+		Filters.RemoveSingle(ToRemoveFilter);
+	}
+}
+
+void USourceFilterCollection::RemoveFilterFromSet(UDataSourceFilter* ToRemoveFilter, UDataSourceFilterSet* FilterSet)
+{
+	FilterSet->Filters.RemoveSingle(ToRemoveFilter);
+
+	// in case set is empty, remove it
+	if (FilterSet->Filters.Num() == 0)
+	{
+		RemoveFilter(FilterSet);
+	}
+}
+
+void USourceFilterCollection::RemoveFilterRecursive(UDataSourceFilter* ToRemoveFilter)
+{
+	ChildToParent.Remove(ToRemoveFilter);
+
+	if (UDataSourceFilterSet* FilterSet = Cast<UDataSourceFilterSet>(ToRemoveFilter))
+	{
+		for (UDataSourceFilter* Filter : FilterSet->Filters)
+		{
+			RemoveFilterRecursive(Filter);
+		}
+	}
+
+	TRACE_FILTER_OPERATION(ToRemoveFilter, ESourceActorFilterOperation::RemoveFilter, FObjectTrace::GetObjectId(Cast<UDataSourceFilterSet>(ToRemoveFilter)));
+}
+
+void USourceFilterCollection::ReplaceFilter(UDataSourceFilter* Destination, UDataSourceFilter* Source)
+{
+	TRACE_FILTER_OPERATION(Source, ESourceActorFilterOperation::ReplaceFilter, FObjectTrace::GetObjectId(Destination));
+
+	UDataSourceFilterSet* OuterFilterSet = ChildToParent.FindChecked(Destination);
+	if (OuterFilterSet)
+	{
+		OuterFilterSet->Modify();
+
+		const int32 FilterEntryIndex = OuterFilterSet->Filters.IndexOfByKey(Destination);
+		OuterFilterSet->Filters[FilterEntryIndex] = Source;
+	}
+	else
+	{
+		const int32 FilterEntryIndex = Filters.IndexOfByKey(Destination);
+		Filters[FilterEntryIndex] = Source;
+	}
+
+	ChildToParent.FindOrAdd(Source) = OuterFilterSet;
+}
+
+void USourceFilterCollection::MoveFilter(UDataSourceFilter* Filter, UDataSourceFilterSet* Destination)
+{
+	TRACE_FILTER_OPERATION(Filter, ESourceActorFilterOperation::MoveFilter, Destination ? FObjectTrace::GetObjectId(Destination) : 0);
+
+	UDataSourceFilterSet*& FilterParent = ChildToParent.FindChecked(Filter);
+
+	if (FilterParent)
+	{
+		RemoveFilterFromSet(Filter, FilterParent);
+	}
+	else
+	{
+		// Currently top level filter
+		Filters.RemoveSingle(Filter);
+	}
+
+	if (Destination)
+	{
+		Destination->Filters.Add(Filter);
+	}
+	else
+	{
+		// Make top level filter
+		Filters.Add(Filter);
+	}
+
+	FilterParent = Destination;
+}
+
+void USourceFilterCollection::Reset()
+{
+	TArray<UDataSourceFilter*> FiltersToRemove = Filters;
+	for (UDataSourceFilter* Filter : FiltersToRemove)
+	{
+		RemoveFilter(Filter);
+	}
+	FilterClasses.Empty();
+	FilterClassMap.Empty();
+}
+
+void USourceFilterCollection::CopyData(USourceFilterCollection* OtherCollection)
+{
+	Filters.Empty();
+	ChildToParent.Empty();
+	FilterClasses = OtherCollection->FilterClasses;
+
+	int32 FilterOffset = 0;
+	for (UDataSourceFilter* Filter : OtherCollection->Filters)
+	{
+		UDataSourceFilter* FilterCopy = RecursiveCopyFilter(Filter, FilterOffset);
+		Filters.Add(FilterCopy);
+		ChildToParent.Add(FilterCopy, nullptr);
+		AddClassName(FilterCopy);
+	}
+	FilterClasses.Empty();
+}
+
+UDataSourceFilter* USourceFilterCollection::RecursiveCopyFilter(UDataSourceFilter* Filter, int32& FilterOffset)
+{
+	UDataSourceFilter* FilterCopy = nullptr;
+	if (Filter)
+	{
+		FilterCopy = DuplicateObject(Filter, this);
+		++FilterOffset;
+
+		if (UDataSourceFilterSet* FilterSet = Cast<UDataSourceFilterSet>(FilterCopy))
+		{
+			TArray<UDataSourceFilter*> NewChildFilters;
+			for (UDataSourceFilter* ChildFilter : FilterSet->Filters)
+			{
+				UDataSourceFilter* ChildFilterCopy = RecursiveCopyFilter(ChildFilter, FilterOffset);
+				ChildToParent.Add(ChildFilterCopy, FilterSet);
+				NewChildFilters.Add(ChildFilterCopy);
+
+				AddClassName(ChildFilterCopy);
+			}
+
+			FilterSet->Filters = NewChildFilters;
+		}
+	}
+	else
+	{
+		UEmptySourceFilter* EmptyFilter = CreateNewFilter<UEmptySourceFilter>();
+		if (FilterClasses.IsValidIndex(FilterOffset))
+		{
+			EmptyFilter->MissingClassName = FilterClasses[FilterOffset];
+		}
+
+		FilterCopy = EmptyFilter;
+		++FilterOffset;
+	}
+
+	return FilterCopy;
+}
+
+void USourceFilterCollection::Serialize(FArchive& Ar)
+{
+	if (Ar.IsSaving() && !Ar.IsTransacting())
+	{
+		FilterClasses.Empty();
+
+		// Regenerate the flat list of class names, corresponding to the entries in Filters, main purpose is to be able to surface missing Filter class names to the user
+		for (UDataSourceFilter* Filter : Filters)
+		{
+			RecursiveRetrieveFilterClassNames(Filter);
+		}
+	}
+
+	Super::Serialize(Ar);
+}
+
+void USourceFilterCollection::AddClassName(UDataSourceFilter* Filter)
+{
+	FilterClassMap.Add(Filter, Filter->GetClass()->GetName());
+}
+
+void USourceFilterCollection::RecursiveRetrieveFilterClassNames(UDataSourceFilter* Filter)
+{
+	if (Filter)
+	{
+		const FString ClassName = FilterClassMap.FindChecked(Filter);
+		FilterClasses.Add(ClassName);
+
+		if (UDataSourceFilterSet* FilterSet = Cast<UDataSourceFilterSet>(Filter))
+		{
+			/* Any child filters will be added in-line, meaning the list will look as such:
+				TopLevelFilterSetA_Class
+					ChildOneFilter_Class
+					ChildTwoFilter_Class
+				TopLevelFilter_Class
+				etc.
+			*/
+			for (UDataSourceFilter* ChildFilter : FilterSet->Filters)
+			{
+				RecursiveRetrieveFilterClassNames(ChildFilter);
+			}
+		}
+	}
+}
\ No newline at end of file
diff --git a/Engine/Plugins/Developer/TraceSourceFiltering/Source/SourceFilteringTrace/Private/SourceFilterManager.cpp b/Engine/Plugins/Developer/TraceSourceFiltering/Source/SourceFilteringTrace/Private/SourceFilterManager.cpp
new file mode 100644
--- /dev/null
+++ b/Engine/Plugins/Developer/TraceSourceFiltering/Source/SourceFilteringTrace/Private/SourceFilterManager.cpp
@@ -0,0 +1,127 @@
+// Copyright Epic Games, Inc. All Rights Reserved.
+
+#include "SourceFilterManager.h"
+
+#include "DataSourceFilter.h"
+#include "UObject/Class.h"
+#include "UObject/UObjectBase.h"
+#include "UObject/Object.h"
+#include "EngineUtils.h"
+#include "GameFramework/Actor.h"
+#include "Engine/World.h"
+#include "ProfilingDebugging/CpuProfilerTrace.h"
+#include "Async/ParallelFor.h"
+
+#include "TraceFilter.h"
+#include "TraceSourceFiltering.h"
+#include "DrawDebugHelpers.h"
+#include "TraceSourceFilteringSettings.h"
+#include "SourceFilterCollection.h"
+
+FSourceFilterManager::FSourceFilterManager(UWorld* InWorld) : World(InWorld)
+{
+	ActorSpawningDelegateHandle = World->AddOnActorSpawnedHandler(FOnActorSpawned::FDelegate::CreateLambda([this](AActor* InActor) 
+	{ 
+		if (CAN_TRACE_OBJECT(World))
+		{
+			ApplyFilters(InActor);
+		}
+	}));
+
+	Settings = FTraceSourceFiltering::Get().GetSettings();
+	FilterCollection = FTraceSourceFiltering::Get().GetFilterCollection();
+}
+
+FSourceFilterManager::~FSourceFilterManager()
+{
+	if (World)
+	{
+		World->RemoveOnActorSpawnedHandler(ActorSpawningDelegateHandle);
+	}
+}
+
+void FSourceFilterManager::Tick(float DeltaTime)
+{
+	TRACE_CPUPROFILER_EVENT_SCOPE(FSourceFilterManager::Tick);
+	if (CAN_TRACE_OBJECT(World))
+	{		
+		for (TActorIterator<AActor> It(World); It; ++It)
+		{
+			if (const AActor* Actor = *It)
+			{
+				ApplyFilters(Actor);
+			}
+		}
+	}
+}
+
+void FSourceFilterManager::ApplyFilters(const AActor* Actor)
+{
+	TRACE_CPUPROFILER_EVENT_SCOPE(FSourceFilterManager::ApplyFilters);
+	const TArray<UDataSourceFilter*>& Filters = FilterCollection->GetFilters();
+	bool bPassesFilters = Filters.Num() ? false : true;
+	bool bAnyEnabledFilter = false;
+
+	FFilterLedger& Ledger = FFilterLedger::Get();
+	Ledger.RejectedFilters.Empty();
+
+	for (const UDataSourceFilter* Filter : Filters)
+	{
+		if (Filter->IsEnabled())
+		{
+			bAnyEnabledFilter = true;
+			const bool bPassesFilter = Filter->DoesActorPassFilter(Actor);
+			bPassesFilters |= bPassesFilter;
+
+			if (!bPassesFilter)
+			{
+				Ledger.RejectedFilters.Add(Filter);
+			}
+		}
+	}
+
+	bPassesFilters = bPassesFilters ? bPassesFilters : !bAnyEnabledFilter;
+	SET_OBJECT_TRACEABLE(Actor, bPassesFilters);
+	
+	// Debug-purpose drawing, allowing users to see impact of Filter set
+	if (Settings->bDrawFilteringStates && (World->WorldType == EWorldType::Game || World->WorldType == EWorldType::PIE))
+	{
+		FVector Origin, Extent;
+		Actor->GetActorBounds(false, Origin, Extent);
+		if (Settings->bDrawOnlyPassingActors)
+		{
+			if (bPassesFilters)
+			{
+				DrawDebugBox(World, Origin, Extent, FQuat::Identity, FColor::Green, false, -1.f, 0, 1.f);
+			}			
+		}
+		else
+		{
+			DrawDebugBox(World, Origin, Extent, FQuat::Identity, bPassesFilters ? FColor::Green : FColor::Red, false, -1.f, 0, 1.f);
+
+			if(!bPassesFilters && Settings->bDrawFilterDescriptionForRejectedActors)
+			{
+				FString ActorString;
+
+				for (const UDataSourceFilter* Filter : Ledger.RejectedFilters)
+				{
+					if (Filter && !Filter->IsA<UDataSourceFilterSet>())
+					{
+						FText Text;
+						Filter->Execute_GetDisplayText(Filter, Text);
+						ActorString += Text.ToString();
+						ActorString += TEXT("\n");
+					}
+				}
+				
+				DrawDebugString(World, FVector::ZeroVector/*Actor->GetActorLocation()*/, ActorString, (AActor*)Actor, FColor::Red, 0.f);
+			}
+		}
+	}
+}
+
+TStatId FSourceFilterManager::GetStatId() const
+{
+	return TStatId();
+}
+
diff --git a/Engine/Plugins/Developer/TraceSourceFiltering/Source/SourceFilteringTrace/Private/SourceFilterManager.h b/Engine/Plugins/Developer/TraceSourceFiltering/Source/SourceFilteringTrace/Private/SourceFilterManager.h
new file mode 100644
--- /dev/null
+++ b/Engine/Plugins/Developer/TraceSourceFiltering/Source/SourceFilteringTrace/Private/SourceFilterManager.h
@@ -0,0 +1,47 @@
+// Copyright Epic Games, Inc. All Rights Reserved.
+
+#pragma once
+
+#include "CoreMinimal.h"
+
+#include "Tickable.h"
+#include "Containers/Array.h"
+#include "Containers/Map.h"
+#include "Delegates/IDelegateInstance.h"
+
+class UDataSourceFilter;
+class UWorld;
+class AActor;
+class UTraceSourceFilteringSettings;
+class USourceFilterCollection;
+
+/** Per-UWorld object that keeps track of the its contained AActor's filtering states */
+class FSourceFilterManager : public FTickableGameObject
+{
+	friend class FTraceSourceFiltering;
+public:
+	FSourceFilterManager(UWorld* InWorld);
+	~FSourceFilterManager();
+
+	/** Begin FTickableGameObject overrides */
+	virtual void Tick(float DeltaTime) override;
+	virtual TStatId GetStatId() const;
+	virtual bool IsTickableInEditor() const { return true; }
+	/** End FTickableGameObject overrides */
+
+protected:
+	/** Applies all UDataSourceFilters to the specified AActor, and updates it filtering state accordingly */
+	void ApplyFilters(const AActor* Actor);
+
+protected:
+	/** Registered delegate for whenever an AActor is spawned within World */
+	FDelegateHandle ActorSpawningDelegateHandle;
+
+	/** Filtering settings for the running instance */
+	const UTraceSourceFilteringSettings* Settings;
+	/** Filter collection, containing the UDataSourceFilters, for the running instance */
+	const USourceFilterCollection* FilterCollection;
+
+	/** UWorld instance that this instance corresponds to */
+	UWorld* World;
+};
diff --git a/Engine/Plugins/Developer/TraceSourceFiltering/Source/SourceFilteringTrace/Private/SourceFilterTrace.cpp b/Engine/Plugins/Developer/TraceSourceFiltering/Source/SourceFilteringTrace/Private/SourceFilterTrace.cpp
new file mode 100644
--- /dev/null
+++ b/Engine/Plugins/Developer/TraceSourceFiltering/Source/SourceFilteringTrace/Private/SourceFilterTrace.cpp
@@ -0,0 +1,246 @@
+// Copyright Epic Games, Inc. All Rights Reserved.
+
+#include "SourceFilterTrace.h"
+
+#if SOURCE_FILTER_TRACE_ENABLED
+
+#include "UObject/UnrealType.h"
+#include "UObject/Class.h"
+#include "Engine/EngineBaseTypes.h"
+
+#include "DataSourceFilter.h"
+#include "TraceWorldFiltering.h"
+#include "EmptySourceFilter.h"
+#include "TraceFilter.h"
+
+#if WITH_EDITOR
+#include "Editor.h"
+#endif // WITH_EDITOR
+
+#define LOCTEXT_NAMESPACE "SourceFilterTrace"
+
+UE_TRACE_CHANNEL_DEFINE(TraceSourceFiltersChannel)
+UE_TRACE_CHANNEL_EXTERN(ImportantChannel)
+
+UE_TRACE_EVENT_BEGIN(SourceFilters, FilterClass)
+	UE_TRACE_EVENT_FIELD(uint64, ClassId)
+UE_TRACE_EVENT_END()
+
+UE_TRACE_EVENT_BEGIN(SourceFilters, FilterInstance)
+	UE_TRACE_EVENT_FIELD(uint64, InstanceId)
+	UE_TRACE_EVENT_FIELD(uint64, ClassId)
+	UE_TRACE_EVENT_FIELD(uint64, SetId)
+UE_TRACE_EVENT_END()
+
+UE_TRACE_EVENT_BEGIN(SourceFilters, FilterSet)
+	UE_TRACE_EVENT_FIELD(uint64, SetId)
+	UE_TRACE_EVENT_FIELD(uint8, Mode)
+UE_TRACE_EVENT_END()
+
+UE_TRACE_EVENT_BEGIN(SourceFilters, FilterOperation)
+	UE_TRACE_EVENT_FIELD(uint64, InstanceId)
+	UE_TRACE_EVENT_FIELD(uint8, Operation)	
+	UE_TRACE_EVENT_FIELD(uint64, Parameter)
+UE_TRACE_EVENT_END()
+
+UE_TRACE_EVENT_BEGIN(SourceFilters, SetFilterSettingValue)
+	UE_TRACE_EVENT_FIELD(uint8, Value)
+UE_TRACE_EVENT_END()
+
+UE_TRACE_EVENT_BEGIN(WorldSourceFilters, WorldInstance)
+	UE_TRACE_EVENT_FIELD(uint64, InstanceId)
+	UE_TRACE_EVENT_FIELD(uint8, Type)
+	UE_TRACE_EVENT_FIELD(uint8, State)
+UE_TRACE_EVENT_END()
+
+UE_TRACE_EVENT_BEGIN(WorldSourceFilters, WorldOperation)
+	UE_TRACE_EVENT_FIELD(uint64, InstanceId)
+	UE_TRACE_EVENT_FIELD(uint8, Operation)
+	UE_TRACE_EVENT_FIELD(uint32, Parameter)
+UE_TRACE_EVENT_END()
+
+TMap<FObjectKey, uint64> FSourceFilterTrace::FilterClassIds;
+TMap<uint64, FObjectKey> FSourceFilterTrace::IDToFilterClass;
+TMap<uint64, FObjectKey> FSourceFilterTrace::IDToFilter;
+TSet<FObjectKey> FSourceFilterTrace::FilterInstances;
+TMap<FString, FObjectKey> FSourceFilterTrace::DataSourceFilterClasses;
+
+TMap<uint64, FObjectKey> FSourceFilterTrace::IDsToWorldInstance;
+
+void FSourceFilterTrace::OutputClass(const TSubclassOf<UDataSourceFilter> InClass)
+{
+	// Skip SKEL and REINST classes.
+	if (InClass->GetName().StartsWith(TEXT("SKEL_")) || InClass->GetName().StartsWith(TEXT("REINST_")))
+	{
+		return;
+	}
+
+	if (InClass->HasAnyClassFlags(CLASS_Abstract | CLASS_HideDropDown | CLASS_Deprecated) && InClass != UEmptySourceFilter::StaticClass())
+	{
+		return;
+	}
+
+	UClass* Class = InClass.Get();
+	TRACE_CLASS(Class);
+
+	DataSourceFilterClasses.Add(Class->GetName(), Class);
+
+	const uint64 Identifier = FObjectTrace::GetObjectId(Class);
+	if (!FilterClassIds.Contains(Class))
+	{
+		FilterClassIds.Add(Class, Identifier);
+		IDToFilterClass.Add(Identifier, Class);
+
+		UE_TRACE_LOG(SourceFilters, FilterClass, TraceSourceFiltersChannel)
+			<< FilterClass.ClassId(Identifier);
+	}
+}
+
+void FSourceFilterTrace::OutputInstance(const UDataSourceFilter* InFilter)
+{
+	if (!FilterInstances.Contains(InFilter))
+	{
+		TRACE_FILTER_CLASS(InFilter->GetClass());
+
+		const uint64 ClassId = FObjectTrace::GetObjectId(InFilter->GetClass());
+		const uint64 InstanceId = FObjectTrace::GetObjectId(InFilter);
+		const uint64 SetId = 0;
+
+		FilterInstances.Add(InFilter);
+		IDToFilter.Add(InstanceId, InFilter);
+
+		FText DisplayText;
+		InFilter->Execute_GetDisplayText(InFilter, DisplayText);
+
+		FString DisplayString = DisplayText.ToString();
+		
+		const uint32 AttachmentSize = (DisplayString.Len() + 1) * sizeof(TCHAR);
+		auto WriteAttachment = [DisplayString, AttachmentSize](uint8* Out)
+		{
+			TCHAR* StringPtr = (TCHAR*)Out;
+			FMemory::Memcpy(StringPtr, *DisplayString, AttachmentSize);
+			StringPtr[AttachmentSize - 1] = '\0';
+		};
+		
+		UE_TRACE_LOG(SourceFilters, FilterInstance, TraceSourceFiltersChannel, AttachmentSize)
+			<< FilterInstance.ClassId(ClassId)
+			<< FilterInstance.InstanceId(InstanceId)
+			<< FilterInstance.SetId(SetId)
+			<< FilterInstance.Attachment(WriteAttachment);
+	}
+}
+
+void FSourceFilterTrace::OutputSet(const UDataSourceFilterSet* InFilterSet)
+{
+	if (InFilterSet)
+	{
+		const uint64 SetId = FObjectTrace::GetObjectId(InFilterSet);
+
+		IDToFilter.Add(SetId, InFilterSet);
+
+		UE_TRACE_LOG(SourceFilters, FilterSet, TraceSourceFiltersChannel)
+			<< FilterSet.SetId(SetId)
+			<< FilterSet.Mode((uint8)InFilterSet->GetFilterSetMode());
+	}
+}
+
+void FSourceFilterTrace::OutputFilterOperation(const UDataSourceFilter* InFilter, ESourceActorFilterOperation Operation, uint64 Parameter)
+{
+	const uint64 InstanceId = FObjectTrace::GetObjectId(InFilter);
+	UE_TRACE_LOG(SourceFilters, FilterOperation, TraceSourceFiltersChannel)
+		<< FilterOperation.InstanceId(InstanceId)
+		<< FilterOperation.Operation((uint8)Operation)
+		<< FilterOperation.Parameter(Parameter);
+}
+
+void FSourceFilterTrace::OutputFilterSettingsValue(const FString& InPropertyName, const uint8 InValue)
+{
+	const uint32 AttachmentSize = (InPropertyName.Len() + 1) * sizeof(TCHAR);
+	auto WriteAttachment = [InPropertyName, AttachmentSize](uint8* Out)
+	{
+		TCHAR* StringPtr = (TCHAR*)Out;
+		FMemory::Memcpy(StringPtr, *InPropertyName, AttachmentSize);
+		StringPtr[AttachmentSize - 1] = '\0';
+	};
+
+	UE_TRACE_LOG(SourceFilters, SetFilterSettingValue, TraceSourceFiltersChannel, AttachmentSize)
+		<< SetFilterSettingValue.Value(InValue)
+		<< SetFilterSettingValue.Attachment(WriteAttachment);
+}
+
+void FSourceFilterTrace::OutputWorld(const UWorld* InWorld)
+{
+	const uint64 InstanceId = FObjectTrace::GetObjectId(InWorld);
+	
+	if (!IDsToWorldInstance.Contains(InstanceId))
+	{
+		IDsToWorldInstance.Add(InstanceId, InWorld);
+
+		FString SendName;
+		FTraceWorldFiltering::GetWorldDisplayString(InWorld, SendName);
+
+		const uint32 AttachmentSize = (SendName.Len() + 1) * sizeof(TCHAR);
+		auto WriteAttachment = [SendName, AttachmentSize](uint8* Out)
+		{
+			TCHAR* StringPtr = (TCHAR*)Out;
+			FMemory::Memcpy(StringPtr, *SendName, AttachmentSize);
+			StringPtr[AttachmentSize - 1] = '\0';
+		};
+
+		UE_TRACE_LOG(WorldSourceFilters, WorldInstance, TraceSourceFiltersChannel, AttachmentSize)
+			<< WorldInstance.InstanceId(InstanceId)
+			<< WorldInstance.Type((uint8)InWorld->WorldType)
+			<< WorldInstance.State(CAN_TRACE_OBJECT(InWorld) ? 1 : 0)
+			<< WorldInstance.Attachment(WriteAttachment);
+	}
+}
+
+void FSourceFilterTrace::OutputWorldOperation(const UWorld* InWorld, EWorldFilterOperation Operation, uint32 Parameter)
+{
+	const uint64 InstanceId = FObjectTrace::GetObjectId(InWorld);
+	UE_TRACE_LOG(WorldSourceFilters, WorldOperation, TraceSourceFiltersChannel)
+		<< WorldOperation.InstanceId(InstanceId)
+		<< WorldOperation.Operation((uint8)Operation)
+		<< WorldOperation.Parameter(Parameter);
+}
+
+template<typename T>
+T* RetrieveTypeById(uint64 Id, const TMap<uint64, FObjectKey>& Map)
+{
+	T* Object = nullptr;
+	if (const FObjectKey* ObjectKey = Map.Find(Id))
+	{
+		Object = Cast<T>(ObjectKey->ResolveObjectPtr());
+	}
+
+	return Object;
+}
+
+UClass* FSourceFilterTrace::RetrieveClassById(uint64 ClassId)
+{
+	return RetrieveTypeById<UClass>(ClassId, IDToFilterClass);
+}
+
+UClass* FSourceFilterTrace::RetrieveClassByName(const FString& ClassName)
+{
+	UClass* Class = nullptr;
+	if (const FObjectKey* ObjectKeyPtr = DataSourceFilterClasses.Find(ClassName))
+	{
+		Class = Cast<UClass>(ObjectKeyPtr->ResolveObjectPtr());
+	}
+	return Class;
+}
+
+UDataSourceFilter* FSourceFilterTrace::RetrieveFilterbyId(uint64 FilterId)
+{
+	return RetrieveTypeById<UDataSourceFilter>(FilterId, IDToFilter);
+}
+
+UWorld* FSourceFilterTrace::RetrieveWorldById(uint64 WorldId)
+{
+	return RetrieveTypeById<UWorld>(WorldId, IDsToWorldInstance);
+}
+
+#endif // SOURCE_FILTER_TRACE_ENABLED
+
+#undef LOCTEXT_NAMESPACE // "SourceFilterTrace"
diff --git a/Engine/Plugins/Developer/TraceSourceFiltering/Source/SourceFilteringTrace/Private/SourceFilteringTraceModule.cpp b/Engine/Plugins/Developer/TraceSourceFiltering/Source/SourceFilteringTrace/Private/SourceFilteringTraceModule.cpp
new file mode 100644
--- /dev/null
+++ b/Engine/Plugins/Developer/TraceSourceFiltering/Source/SourceFilteringTrace/Private/SourceFilteringTraceModule.cpp
@@ -0,0 +1,101 @@
+// Copyright Epic Games, Inc. All Rights Reserved.
+
+#include "SourceFilteringTraceModule.h"
+#include "Modules/ModuleManager.h"
+#include "Modules/ModuleInterface.h"
+#include "Trace/Trace.h"
+#include "UObject/SoftObjectPtr.h"
+#include "UObject/SoftObjectPath.h"
+#include "UObject/Class.h"
+#include "UObject/UnrealType.h"
+#include "UObject/UObjectIterator.h"
+#include "Blueprint/BlueprintSupport.h"
+#include "HAL/IConsoleManager.h"
+#include "PropertyPathHelpers.h"
+#include "AssetRegistryModule.h"
+#include "Engine/EngineTypes.h"
+#include "Trace/Detail/Channel.h"
+
+#include "DataSourceFilter.h"
+#include "SourceFilterManager.h"
+#include "DataSourceFilterSet.h"
+#include "SourceFilterTrace.h"
+#include "TraceWorldFiltering.h"
+#include "TraceSourceFiltering.h"
+
+#if WITH_EDITOR
+#include "Kismet2/KismetEditorUtilities.h"
+#endif // WITH_EDITOR
+
+void FSourceFilteringTraceModule::TraceFilterClasses()
+{
+	FAssetRegistryModule& AssetRegistryModule = FModuleManager::LoadModuleChecked<FAssetRegistryModule>(TEXT("AssetRegistry"));
+
+	TArray<FAssetData> Filters;
+	TMultiMap<FName, FString> TagValues = { { FBlueprintTags::NativeParentClassPath, FString::Printf(TEXT("%s'%s'"), *UClass::StaticClass()->GetName(), *UDataSourceFilter::StaticClass()->GetPathName()) } };
+	AssetRegistryModule.Get().GetAssetsByTagValues(TagValues, Filters);
+
+	for (const FAssetData& AssetData : Filters)
+	{
+		FAssetDataTagMapSharedView::FFindTagResult GeneratedClassPathPtr = AssetData.TagsAndValues.FindTag(TEXT("GeneratedClass"));
+		if (GeneratedClassPathPtr.IsSet())
+		{
+			const FString BPClassObjectPath = FPackageName::ExportTextPathToObjectPath(GeneratedClassPathPtr.GetValue());
+			const FString BPClassName = FPackageName::ObjectPathToObjectName(BPClassObjectPath);
+
+			UClass* ChildClass = StaticLoadClass(UDataSourceFilter::StaticClass(), nullptr, *BPClassObjectPath);
+			TRACE_FILTER_CLASS(ChildClass);
+		}
+	}
+
+	for (TObjectIterator<UClass> ClassIt; ClassIt; ++ClassIt)
+	{
+		if (ClassIt->IsChildOf(UDataSourceFilter::StaticClass()))
+		{
+			if (UClass* FilterClass = *ClassIt)
+			{
+				if (FilterClass != UDataSourceFilter::StaticClass() && FilterClass != UDataSourceFilterSet::StaticClass())
+				{
+					TRACE_FILTER_CLASS(FilterClass);
+				}
+			}
+		}
+	}
+}
+#if WITH_EDITOR
+void FSourceFilteringTraceModule::HandleNewFilterBlueprintCreated(UBlueprint* InBlueprint)
+{
+	TraceFilterClasses();
+}
+#endif // WITH_EDITOR
+
+void FSourceFilteringTraceModule::StartupModule()
+{
+	FTraceWorldFiltering::Initialize();
+	FTraceSourceFiltering::Initialize();
+
+	// Forcefully enable the source trace channel
+	Trace::FChannel::Toggle(&TraceSourceFiltersChannel, true);
+
+#if WITH_EDITOR
+	// Add callback to trace out Filter Classes once the Asset Registry has finished loading 
+	FAssetRegistryModule& AssetRegistryModule = FModuleManager::LoadModuleChecked<FAssetRegistryModule>(TEXT("AssetRegistry"));
+	AssetRegistryModule.Get().OnFilesLoaded().AddRaw(this, &FSourceFilteringTraceModule::TraceFilterClasses);
+
+	// Add callback to trace out Filter Classes when a new UDataSourceFilter-based Blueprint has been created 
+	FKismetEditorUtilities::RegisterOnBlueprintCreatedCallback(this, UDataSourceFilter::StaticClass(), FKismetEditorUtilities::FOnBlueprintCreated::CreateRaw(this, &FSourceFilteringTraceModule::HandleNewFilterBlueprintCreated));
+#endif // WITH_EDITOR	
+
+	TraceFilterClasses();
+}
+
+void FSourceFilteringTraceModule::ShutdownModule()
+{
+#if WITH_EDITOR
+	FKismetEditorUtilities::UnregisterAutoBlueprintNodeCreation(this);
+#endif // WITH_EDITOR
+	FTraceWorldFiltering::Destroy();
+}
+
+
+IMPLEMENT_MODULE(FSourceFilteringTraceModule, SourceFilteringTrace);
diff --git a/Engine/Plugins/Developer/TraceSourceFiltering/Source/SourceFilteringTrace/Private/SourceFilteringTraceModule.h b/Engine/Plugins/Developer/TraceSourceFiltering/Source/SourceFilteringTrace/Private/SourceFilteringTraceModule.h
new file mode 100644
--- /dev/null
+++ b/Engine/Plugins/Developer/TraceSourceFiltering/Source/SourceFilteringTrace/Private/SourceFilteringTraceModule.h
@@ -0,0 +1,20 @@
+// Copyright Epic Games, Inc. All Rights Reserved.
+
+#pragma once
+
+#include "CoreMinimal.h"
+#include "Modules/ModuleInterface.h"
+
+#include "SourceFilterCollection.h"
+
+class FSourceFilteringTraceModule : public IModuleInterface
+{
+public:
+	virtual void StartupModule() override;
+	virtual void ShutdownModule() override;
+protected:
+	void TraceFilterClasses();
+#if WITH_EDITOR	
+	void HandleNewFilterBlueprintCreated(class UBlueprint* InBlueprint);
+#endif // WITH_EDITOR
+};
diff --git a/Engine/Plugins/Developer/TraceSourceFiltering/Source/SourceFilteringTrace/Private/TraceSourceFiltering.cpp b/Engine/Plugins/Developer/TraceSourceFiltering/Source/SourceFilteringTrace/Private/TraceSourceFiltering.cpp
new file mode 100644
--- /dev/null
+++ b/Engine/Plugins/Developer/TraceSourceFiltering/Source/SourceFilteringTrace/Private/TraceSourceFiltering.cpp
@@ -0,0 +1,222 @@
+// Copyright Epic Games, Inc. All Rights Reserved.
+
+#include "TraceSourceFiltering.h"
+#include "PropertyPathHelpers.h"
+#include "SourceFilterCollection.h"
+#include "TraceSourceFilteringSettings.h"
+#include "SourceFilterTrace.h"
+#include "TraceWorldFiltering.h"
+
+FTraceSourceFiltering::FTraceSourceFiltering()
+{
+	FilterCollection = NewObject<USourceFilterCollection>(GetTransientPackage(), NAME_None, RF_Transactional | RF_Transient);
+	Settings = DuplicateObject(GetDefault<UTraceSourceFilteringSettings>(), nullptr);
+
+	PopulateRemoteTraceCommands();
+}
+
+void FTraceSourceFiltering::Initialize()
+{
+	FTraceSourceFiltering::Get();
+}
+
+FTraceSourceFiltering& FTraceSourceFiltering::Get()
+{
+	static FTraceSourceFiltering SourceFiltering;
+	return SourceFiltering;
+}
+
+USourceFilterCollection* FTraceSourceFiltering::GetFilterCollection()
+{
+	return FilterCollection;
+}
+
+UTraceSourceFilteringSettings* FTraceSourceFiltering::GetSettings()
+{
+	return Settings;
+}
+
+static UClass* ConvertArgumentToClass(FString Argument)
+{
+	const uint64 ClassId = FCString::Atoi64(*Argument);
+
+	UClass* Class = FSourceFilterTrace::RetrieveClassById(ClassId);
+	checkf(Class != nullptr, TEXT("Failed to retrieve Class from command argument"));
+	return Class;
+}
+
+static UDataSourceFilterSet* ConvertArgumentToFilterSet(FString Argument)
+{
+	const uint64 SetId = FCString::Atoi64(*Argument);
+	// Cast as it is valid to specify a null Filter Set as a target for MoveFilter operation
+	UDataSourceFilterSet* FilterSet = Cast<UDataSourceFilterSet>(FSourceFilterTrace::RetrieveFilterbyId(SetId));
+
+	return FilterSet;
+}
+
+static EFilterSetMode ConvertArgumentToSetMode(FString Argument)
+{
+	const EFilterSetMode Mode = (EFilterSetMode)FCString::Atoi(*Argument);
+	return Mode;
+}
+
+static UDataSourceFilter* ConvertArgumentToFilter(FString Argument)
+{
+	const uint64 InstanceId = FCString::Atoi64(*Argument);
+
+	UDataSourceFilter* Filter = FSourceFilterTrace::RetrieveFilterbyId(InstanceId);
+	ensure(Filter != nullptr);
+
+	return Filter;
+}
+
+static UWorld* ConvertArgumentToWorld(FString Argument)
+{
+	const uint64 WorldId = FCString::Atoi64(*Argument);
+
+	UWorld* World = FSourceFilterTrace::RetrieveWorldById(WorldId);
+	ensure(World != nullptr);
+
+	return World;
+}
+
+void FTraceSourceFiltering::ProcessRemoteCommand(const FString& Command, const TArray<FString>& Arguments)
+{
+	if (FFilterCommand* FilterCommand = CommandMap.Find(Command))
+	{
+		ensure(FilterCommand->NumExpectedArguments == Arguments.Num());
+		FilterCommand->Function(Arguments);
+	}
+}
+
+void FTraceSourceFiltering::AddReferencedObjects(FReferenceCollector& Collector)
+{
+	Collector.AddReferencedObject(FilterCollection);
+	Collector.AddReferencedObject(Settings);
+}
+
+void FTraceSourceFiltering::PopulateRemoteTraceCommands()
+{
+	CommandMap.Add(TEXT("AddFilterById"), { [this](const TArray<FString>& Arguments) -> void { FilterCollection->AddFilterOfClass(ConvertArgumentToClass(Arguments[0])); }, 1 });
+
+	CommandMap.Add(TEXT("AddFilterClassToSet"),
+		{ [this](const TArray<FString>& Arguments) -> void { FilterCollection->AddFilterOfClassToSet(ConvertArgumentToClass(Arguments[0]), ConvertArgumentToFilterSet(Arguments[1])); }, 2 });
+
+	CommandMap.Add(TEXT("AddFilterToSet"),
+		{ [this](const TArray<FString>& Arguments) -> void { FilterCollection->MoveFilter(ConvertArgumentToFilter(Arguments[0]), ConvertArgumentToFilterSet(Arguments[1])); }, 2 });
+
+	CommandMap.Add(TEXT("MakeFilterSet"),
+		{ [this](const TArray<FString>& Arguments) -> void { FilterCollection->MakeFilterSet(ConvertArgumentToFilter(Arguments[0]), ConvertArgumentToFilter(Arguments[1]), EFilterSetMode::AND); }, 2 });
+
+	CommandMap.Add(TEXT("MakeSingleFilterSet"),
+		{ [this](const TArray<FString>& Arguments) -> void { FilterCollection->ConvertFilterToSet(ConvertArgumentToFilter(Arguments[0]), ConvertArgumentToSetMode(Arguments[1])); }, 2 });
+
+	CommandMap.Add(TEXT("ResetFilters"),
+		{ [this](const TArray<FString>& Arguments) -> void { FilterCollection->Reset(); }, 0 });
+
+	CommandMap.Add(TEXT("MoveFilter"),
+		{ [this](const TArray<FString>& Arguments) -> void { FilterCollection->MoveFilter(ConvertArgumentToFilter(Arguments[0]), ConvertArgumentToFilterSet(Arguments[1])); }, 2 });
+
+	CommandMap.Add(TEXT("RemoveFilter"),
+		{ [this](const TArray<FString>& Arguments) -> void { FilterCollection->RemoveFilter(ConvertArgumentToFilter(Arguments[0])); }, 1 });
+
+	CommandMap.Add(TEXT("SetFilterMode"),
+		{ [](const TArray<FString>& Arguments) -> void
+			{
+				UDataSourceFilterSet* FilterSet = ConvertArgumentToFilterSet(Arguments[0]);
+				EFilterSetMode Mode = ConvertArgumentToSetMode(Arguments[1]);
+				FilterSet->SetFilterMode(Mode);
+			},
+			2
+		}
+	);
+
+	CommandMap.Add(TEXT("SetFilterState"),
+		{ [](const TArray<FString>& Arguments) -> void
+			{
+				UDataSourceFilter* Filter = ConvertArgumentToFilter(Arguments[0]);
+				const bool bState = (bool)FCString::Atoi(*Arguments[1]);
+				Filter->SetEnabled(bState);
+			},
+			2
+		}
+	);
+
+	CommandMap.Add(TEXT("SetFilterSetting"),
+		{ [this](const TArray<FString>& Arguments) -> void
+			{
+				FCachedPropertyPath PropertyPath(Arguments[0]);
+				const bool bValue = (bool)FCString::Atoi(*Arguments[1]);
+				ensure(PropertyPathHelpers::SetPropertyValue(Settings, PropertyPath, bValue));
+				TRACE_FILTER_SETTINGS_VALUE(PropertyPath.ToString(), (uint8)bValue);
+			},
+			2
+		}
+	);
+
+	CommandMap.Add(TEXT("SetFiltersFromPreset"),
+		{ [this](const TArray<FString>& Arguments) -> void
+			{
+				FString PresetString = Arguments[0];
+
+				// Parse out each class + parent index entry from the preset, each entry is stored as ClassName|ParentIndex
+				TArray<FString> ClassIndexStrings;
+				PresetString.ParseIntoArray(ClassIndexStrings, TEXT(";"));
+
+				TArray<FString> ClassNames;
+				TMap<int32, int32> ChildToParent;
+				for (const FString& ClassIndexString : ClassIndexStrings)
+				{
+					// Separate out the class name from the index
+					TArray<FString> Strings;
+					ClassIndexString.ParseIntoArray(Strings, TEXT("|"));
+
+					const FString ClassName = Strings[0];
+					const int32 ParentIndex = FCString::Atoi(*Strings[1]);
+
+					// Generate mapping from Child to Parent class instances
+					const int32 FilterIndex = ClassNames.Add(ClassName);
+					ChildToParent.Add(FilterIndex, ParentIndex);
+				}
+
+				FilterCollection->AddFiltersFromPreset(ClassNames, ChildToParent);
+			},
+			1
+		}
+	);
+
+	CommandMap.Add(TEXT("SetWorldTypeFilterState"),
+		{ [](const TArray<FString>& Arguments) -> void
+			{
+				const EWorldType::Type WorldType = (EWorldType::Type)FCString::Atoi(*Arguments[0]);
+				const bool bState = (bool)FCString::Atoi(*Arguments[1]);
+
+				FTraceWorldFiltering::SetStateByWorldType(WorldType, bState);
+			},
+			2
+		}
+	);
+
+	CommandMap.Add(TEXT("SetWorldNetModeFilterState"),
+		{ [](const TArray<FString>& Arguments) -> void
+			{
+				const ENetMode NetMode = (ENetMode)FCString::Atoi(*Arguments[0]);
+				const bool bState = (bool)FCString::Atoi(*Arguments[1]);
+
+				FTraceWorldFiltering::SetStateByWorldNetMode(NetMode, bState);
+			},
+			2
+		}
+	);
+
+	CommandMap.Add(TEXT("SetWorldFilterState"),
+		{ [](const TArray<FString>& Arguments) -> void
+			{
+				const UWorld* World = ConvertArgumentToWorld(Arguments[0]);
+				const bool bState = (bool)FCString::Atoi(*Arguments[1]);
+				FTraceWorldFiltering::SetWorldState(World, bState);
+			},
+			2
+		}
+	);
+}
diff --git a/Engine/Plugins/Developer/TraceSourceFiltering/Source/SourceFilteringTrace/Private/TraceSourceFilteringProjectSettings.h b/Engine/Plugins/Developer/TraceSourceFiltering/Source/SourceFilteringTrace/Private/TraceSourceFilteringProjectSettings.h
new file mode 100644
--- /dev/null
+++ b/Engine/Plugins/Developer/TraceSourceFiltering/Source/SourceFilteringTrace/Private/TraceSourceFilteringProjectSettings.h
@@ -0,0 +1,20 @@
+// Copyright Epic Games, Inc. All Rights Reserved.
+
+#pragma once
+
+#include "Engine/DeveloperSettings.h"
+#include "UObject/SoftObjectPtr.h"
+
+#include "DataSourceFilter.h"
+
+#include "TraceSourceFilteringProjectSettings.generated.h"
+
+UCLASS(config = Engine, meta = (DisplayName = "Trace Source Filtering"), defaultconfig)
+class UTraceSourceFilteringProjectSettings : public UDeveloperSettings
+{
+	GENERATED_BODY()
+
+public:
+	UPROPERTY(config, EditAnywhere, Category = TraceSourceFiltering, AdvancedDisplay, meta = (DisplayName = "Source Filter Classes, which should be incorporated into the cook", RelativeToGameContentDir))
+	TArray<TSoftClassPtr<UDataSourceFilter>> CookedSourceFilterClasses;
+};
diff --git a/Engine/Plugins/Developer/TraceSourceFiltering/Source/SourceFilteringTrace/Private/TraceWorldFiltering.cpp b/Engine/Plugins/Developer/TraceSourceFiltering/Source/SourceFilteringTrace/Private/TraceWorldFiltering.cpp
new file mode 100644
--- /dev/null
+++ b/Engine/Plugins/Developer/TraceSourceFiltering/Source/SourceFilteringTrace/Private/TraceWorldFiltering.cpp
@@ -0,0 +1,252 @@
+// Copyright Epic Games, Inc. All Rights Reserved.
+
+#include "TraceWorldFiltering.h"
+#include "SourceFilterTrace.h"
+#include "Misc/ScopeLock.h"
+#include "TraceFilters.h"
+#include "TraceFilter.h"
+
+#include "ObjectTrace.h"
+#include "SourceFilterManager.h"
+
+#if WITH_EDITOR
+#include "Editor.h"
+#endif // WITH_EDITOR
+
+FDelegateHandle FTraceWorldFiltering::WorldInitHandle;
+FDelegateHandle FTraceWorldFiltering::WorldPostInitHandle;
+FDelegateHandle FTraceWorldFiltering::WorldBeginTearDownHandle;
+FDelegateHandle FTraceWorldFiltering::WorldCleanupHandle;
+FDelegateHandle FTraceWorldFiltering::PreWorldFinishDestroyHandle;
+
+FTraceWorldFiltering::FTraceWorldFilterStateChanged FTraceWorldFiltering::FilterStateChangedDelegate;
+
+TArray<const UWorld*> FTraceWorldFiltering::Worlds;
+TMap<EWorldType::Type, bool> FTraceWorldFiltering::WorldTypeFilterStates;
+TMap<ENetMode, bool> FTraceWorldFiltering::NetModeFilterStates;
+
+FCriticalSection FTraceWorldFiltering::WorldFilterStatesCritical;
+TMap<const UWorld*, FSourceFilterManager*> FTraceWorldFiltering::WorldSourceFilterManagers;
+
+#define LOCTEXT_NAMESPACE "FTraceWorldFilter" 
+
+void FTraceWorldFiltering::Initialize()
+{
+	WorldInitHandle = FWorldDelegates::OnPreWorldInitialization.AddStatic(&FTraceWorldFiltering::OnWorldInit);
+	WorldPostInitHandle = FWorldDelegates::OnPostWorldInitialization.AddStatic(&FTraceWorldFiltering::OnWorldPostInit);
+	WorldCleanupHandle = FWorldDelegates::OnWorldCleanup.AddStatic(&FTraceWorldFiltering::OnWorldCleanup);
+	WorldBeginTearDownHandle = FWorldDelegates::OnWorldBeginTearDown.AddStatic(&FTraceWorldFiltering::RemoveWorld);
+	PreWorldFinishDestroyHandle = FWorldDelegates::OnPreWorldFinishDestroy.AddStatic(&FTraceWorldFiltering::RemoveWorld);
+
+	// Disable engine level world filter
+	DISABLE_ENGINE_WORLD_TRACE_FILTERING();
+}
+
+void FTraceWorldFiltering::Destroy()
+{
+	FWorldDelegates::OnPreWorldInitialization.Remove(WorldInitHandle);
+	FWorldDelegates::OnWorldCleanup.Remove(WorldCleanupHandle);
+	FWorldDelegates::OnWorldBeginTearDown.Remove(WorldBeginTearDownHandle);
+	FWorldDelegates::OnPreWorldFinishDestroy.Remove(PreWorldFinishDestroyHandle);
+
+	for (const UWorld* World : Worlds)
+	{
+		if (World)
+		{
+			SetWorldState(World, true);
+		}
+	}
+}
+
+const FSourceFilterManager* FTraceWorldFiltering::GetWorldSourceFilterManager(const UWorld* World)
+{
+	return WorldSourceFilterManagers.FindRef(World);
+}
+
+const TArray<const UWorld*>& FTraceWorldFiltering::GetWorlds()
+{
+	return Worlds;
+}
+
+bool FTraceWorldFiltering::IsWorldTypeTraceable(EWorldType::Type InType)
+{
+	FScopeLock Lock(&WorldFilterStatesCritical);
+	const bool* bStatePtr = WorldTypeFilterStates.Find(InType);
+	return bStatePtr ? *bStatePtr : true;
+}
+
+bool FTraceWorldFiltering::IsWorldNetModeTraceable(ENetMode InNetMode)
+{
+	FScopeLock Lock(&WorldFilterStatesCritical);
+	const bool* bStatePtr = NetModeFilterStates.Find(InNetMode);
+	return bStatePtr ? *bStatePtr : true;
+}
+
+void FTraceWorldFiltering::SetStateByWorldType(EWorldType::Type WorldType, bool bState)
+{
+	// Update tracked flag 
+	{
+		FScopeLock Lock(&WorldFilterStatesCritical);
+		WorldTypeFilterStates.FindOrAdd(WorldType) = bState;
+	}
+	UpdateWorldFiltering();
+
+	// Trace out event for filtering change 
+	uint32 Parameter = 0;
+	Parameter |= ((uint32)WorldType << 16);
+	Parameter |= (uint32)bState;
+
+	OnFilterStateChanged().Broadcast();
+	TRACE_WORLD_OPERATION(nullptr, EWorldFilterOperation::TypeFilter, Parameter);
+}
+
+void FTraceWorldFiltering::SetStateByWorldNetMode(ENetMode NetMode, bool bState)
+{
+	// Update tracked flag 
+	{
+		FScopeLock Lock(&WorldFilterStatesCritical);
+		NetModeFilterStates.FindOrAdd(NetMode) = bState;
+	}
+	UpdateWorldFiltering();
+
+	// Trace out event for filtering change 
+	uint32 Parameter = 0;
+	Parameter |= ((uint32)NetMode << 16);
+	Parameter |= (uint32)bState;
+
+	OnFilterStateChanged().Broadcast();
+	TRACE_WORLD_OPERATION(nullptr, EWorldFilterOperation::NetModeFilter, Parameter);
+}
+
+void FTraceWorldFiltering::SetWorldState(const UWorld* InWorld, bool bState)
+{
+	// Mark specific UWorld as traceable
+	SET_OBJECT_TRACEABLE(InWorld, bState);
+
+	// Make sure to trace out contained ULevels when the provided World is marked as traceable
+	if (bState)
+	{
+		for (FConstLevelIterator Iterator = InWorld->GetLevelIterator(); Iterator; ++Iterator)
+		{
+			TRACE_OBJECT(*Iterator);
+		}
+	}
+
+
+	// Trace out event for world's filtering state 
+	OnFilterStateChanged().Broadcast();
+	TRACE_WORLD_OPERATION(InWorld, EWorldFilterOperation::InstanceFilter, bState);
+}
+
+void FTraceWorldFiltering::GetWorldDisplayString(const UWorld* InWorld, FString& OutDisplayString)
+{
+	const FString WorldName = InWorld->GetName();
+
+	FString PostFix;
+	if (InWorld->WorldType == EWorldType::PIE)
+	{
+		switch (InWorld->GetNetMode())
+		{
+		case NM_Client:
+			PostFix = FString::Printf(TEXT("Client %i"), InWorld->GetOutermost()->PIEInstanceID - 1);
+			break;
+		case NM_DedicatedServer:
+		case NM_ListenServer:
+			PostFix = LOCTEXT("ServerPostfix", "Server").ToString();
+			break;
+#if WITH_EDITOR
+		case NM_Standalone:
+			PostFix = GEditor->bIsSimulatingInEditor ? LOCTEXT("SimulateInEditorPostfix", "Simulate").ToString() : LOCTEXT("PlayInEditorPostfix", "PIE").ToString();
+			break;
+#endif // WITH_EDITOR
+		}
+	}
+	else if (InWorld->WorldType == EWorldType::Editor || InWorld->WorldType == EWorldType::EditorPreview)
+	{
+		PostFix = LOCTEXT("EditorPostfix", "Editor").ToString();
+	}
+	else if (InWorld->WorldType == EWorldType::Game || InWorld->WorldType == EWorldType::GamePreview || InWorld->WorldType == EWorldType::GameRPC)
+	{
+		PostFix = LOCTEXT("GamePostfix", "Game").ToString();
+	}
+
+	OutDisplayString = FString::Printf(TEXT("%s [%s]"), *WorldName, *PostFix);
+}
+
+FTraceWorldFiltering::FTraceWorldFilterStateChanged& FTraceWorldFiltering::OnFilterStateChanged()
+{
+	return FilterStateChangedDelegate;
+}
+
+void FTraceWorldFiltering::OnWorldInit(UWorld* InWorld, const UWorld::InitializationValues IVS)
+{
+	if (InWorld && !Worlds.Contains(InWorld))
+	{
+		FScopeLock Lock(&FTraceWorldFiltering::WorldFilterStatesCritical);
+
+		// Setup the initial filtering state according to set world properties
+		const bool* bWorldTypeStatePtr = WorldTypeFilterStates.Find(InWorld->WorldType);
+		const bool* bNetModeStatePtr = NetModeFilterStates.Find(InWorld->GetNetMode());
+
+		const bool bTraceable = (bWorldTypeStatePtr ? *bWorldTypeStatePtr : true) && (bNetModeStatePtr ? *bNetModeStatePtr : true);
+		SET_OBJECT_TRACEABLE(InWorld, bTraceable);
+
+		// Create source filter manager for this world 
+		FSourceFilterManager* Manager = new FSourceFilterManager(InWorld);
+		WorldSourceFilterManagers.Add(InWorld, Manager);
+
+		Worlds.Add(InWorld);
+
+		OnFilterStateChanged().Broadcast();
+
+		// Trace out new UWorld instance
+		TRACE_WORLD_INSTANCE(InWorld);
+	}
+}
+
+void FTraceWorldFiltering::OnWorldPostInit(UWorld* World, const UWorld::InitializationValues IVS)
+{
+	// Handle all loaded AActors within this world, to handle duplicated UWorlds (PIE)
+	FSourceFilterManager* Manager = WorldSourceFilterManagers.FindChecked(World);
+	Manager->Tick(0.f);
+}
+
+void FTraceWorldFiltering::OnWorldCleanup(UWorld* InWorld, bool bSessionEnded, bool bCleanupResources)
+{
+	RemoveWorld(InWorld);
+}
+
+void FTraceWorldFiltering::RemoveWorld(UWorld* InWorld)
+{
+	// Cleanup data for provided UWorld 
+	if (Worlds.Remove(InWorld))
+	{
+		FSourceFilterManager* Manager = nullptr;
+		if (WorldSourceFilterManagers.RemoveAndCopyValue(InWorld, Manager))
+		{
+			delete Manager;
+		}
+
+		TRACE_WORLD_OPERATION(InWorld, EWorldFilterOperation::RemoveWorld, 0);
+
+		OnFilterStateChanged().Broadcast();
+	}
+}
+
+void FTraceWorldFiltering::UpdateWorldFiltering()
+{
+	FScopeLock Lock(&FTraceWorldFiltering::WorldFilterStatesCritical);
+
+	// Update all active worlds to new filtering state
+	for (const UWorld* World : Worlds)
+	{
+		// Setup the initial filtering state according to set world properties
+		const bool* bWorldTypeStatePtr = WorldTypeFilterStates.Find(World->WorldType);
+		const bool* bNetModeStatePtr = NetModeFilterStates.Find(World->GetNetMode());
+
+		const bool bTraceable = (bWorldTypeStatePtr ? *bWorldTypeStatePtr : true) && (bNetModeStatePtr ? *bNetModeStatePtr : true);
+		SetWorldState(World, bTraceable);
+	}
+}
+
+#undef LOCTEXT_NAMESPACE // "FTraceWorldFilter" 
\ No newline at end of file
diff --git a/Engine/Plugins/Developer/TraceSourceFiltering/Source/SourceFilteringTrace/Public/DataSourceFilter.cpp b/Engine/Plugins/Developer/TraceSourceFiltering/Source/SourceFilteringTrace/Public/DataSourceFilter.cpp
new file mode 100644
--- /dev/null
+++ b/Engine/Plugins/Developer/TraceSourceFiltering/Source/SourceFilteringTrace/Public/DataSourceFilter.cpp
@@ -0,0 +1,43 @@
+// Copyright Epic Games, Inc. All Rights Reserved.
+
+#include "DataSourceFilter.h"
+#include "SourceFilterTrace.h"
+
+UDataSourceFilter::UDataSourceFilter() : bIsEnabled(true)
+{
+
+}
+
+UDataSourceFilter::~UDataSourceFilter()
+{
+
+}
+
+bool UDataSourceFilter::DoesActorPassFilter_Implementation(const AActor* InActor) const
+{
+	return DoesActorPassFilter_Internal(InActor);
+}
+
+void UDataSourceFilter::SetEnabled(bool bState)
+{
+	TRACE_FILTER_OPERATION(this, ESourceActorFilterOperation::SetFilterState, (uint64)bState);
+	bIsEnabled = bState;
+}
+
+bool UDataSourceFilter::IsEnabled() const
+{
+	return bIsEnabled;
+}
+
+void UDataSourceFilter::GetDisplayText_Internal(FText& OutDisplayText) const
+{
+	if (UClass* Class = GetClass())
+	{
+		OutDisplayText = FText::FromString(GetClass()->GetName());
+	}
+}
+
+bool UDataSourceFilter::DoesActorPassFilter_Internal(const AActor* InActor) const
+{
+	return true;
+}
diff --git a/Engine/Plugins/Developer/TraceSourceFiltering/Source/SourceFilteringTrace/Public/DataSourceFilter.h b/Engine/Plugins/Developer/TraceSourceFiltering/Source/SourceFilteringTrace/Public/DataSourceFilter.h
new file mode 100644
--- /dev/null
+++ b/Engine/Plugins/Developer/TraceSourceFiltering/Source/SourceFilteringTrace/Public/DataSourceFilter.h
@@ -0,0 +1,42 @@
+// Copyright Epic Games, Inc. All Rights Reserved.
+
+#pragma once
+
+#include "CoreMinimal.h"
+#include "UObject/Object.h"
+#include "GameFramework/Actor.h"
+#include "IDataSourceFilterInterface.h"
+
+#include "DataSourceFilter.generated.h"
+
+struct FFilterLedger : public TThreadSingleton<FFilterLedger>
+{
+public:
+	TArray<const UDataSourceFilter*> RejectedFilters;
+};
+
+UCLASS(Blueprintable)
+class SOURCEFILTERINGTRACE_API UDataSourceFilter : public UObject, public IDataSourceFilterInterface
+{
+	GENERATED_BODY()
+public:
+	UDataSourceFilter();
+	virtual ~UDataSourceFilter();
+
+	UFUNCTION(BlueprintNativeEvent, Category = TraceSourceFiltering)
+	bool DoesActorPassFilter(const AActor* InActor) const;
+	bool DoesActorPassFilter_Implementation(const AActor* InActor) const;
+
+	/** Begin IDataSourceFilterInterface overrides */
+	virtual void SetEnabled(bool bState) override;	
+	virtual bool IsEnabled() const final;
+protected:
+	virtual void GetDisplayText_Internal(FText& OutDisplayText) const override;
+	/** End IDataSourceFilterInterface overrides */
+
+	virtual bool DoesActorPassFilter_Internal(const AActor* InActor) const;
+protected:
+	/** Whether or not this filter is enabled */
+	bool bIsEnabled;
+};
+
diff --git a/Engine/Plugins/Developer/TraceSourceFiltering/Source/SourceFilteringTrace/Public/DataSourceFilterSet.h b/Engine/Plugins/Developer/TraceSourceFiltering/Source/SourceFilteringTrace/Public/DataSourceFilterSet.h
new file mode 100644
--- /dev/null
+++ b/Engine/Plugins/Developer/TraceSourceFiltering/Source/SourceFilteringTrace/Public/DataSourceFilterSet.h
@@ -0,0 +1,40 @@
+// Copyright Epic Games, Inc. All Rights Reserved.
+
+#pragma once
+
+#include "DataSourceFilter.h"
+#include "IDataSourceFilterSetInterface.h"
+
+#include "DataSourceFilterSet.generated.h"
+
+/** Engine implementation of IDataSourceFilterSetInterface */
+UCLASS(NotBlueprintable)
+class SOURCEFILTERINGTRACE_API UDataSourceFilterSet : public UDataSourceFilter, public IDataSourceFilterSetInterface
+{
+	friend class USourceFilterCollection;
+
+	GENERATED_BODY()
+public:	
+	const TArray<UDataSourceFilter*>& GetFilters() const { return Filters; }
+	void SetFilterMode(EFilterSetMode InMode);
+
+	/** Begin IDataSourceFilterSetInterface overrides */
+	virtual EFilterSetMode GetFilterSetMode() const override;
+	/** Begin IDataSourceFilterSetInterface overrides */
+
+	/** Begin UDataSourceFilter overrides */
+	virtual void SetEnabled(bool bState) override;
+protected:
+	virtual bool DoesActorPassFilter_Internal(const AActor* InActor) const override;
+	virtual void GetDisplayText_Internal(FText& OutDisplayText) const override;
+	/** End UDataSourceFilter overrides */
+
+protected:
+	/** Contained Filter instance */
+	UPROPERTY()
+	TArray<UDataSourceFilter*> Filters;
+
+	/** Current Filter set operation */
+	UPROPERTY()
+	EFilterSetMode Mode;
+};
diff --git a/Engine/Plugins/Developer/TraceSourceFiltering/Source/SourceFilteringTrace/Public/EmptySourceFilter.h b/Engine/Plugins/Developer/TraceSourceFiltering/Source/SourceFilteringTrace/Public/EmptySourceFilter.h
new file mode 100644
--- /dev/null
+++ b/Engine/Plugins/Developer/TraceSourceFiltering/Source/SourceFilteringTrace/Public/EmptySourceFilter.h
@@ -0,0 +1,29 @@
+// Copyyright 1998 - 2019 Epic Games, Inc.All Rights Reserved.
+
+#pragma once
+
+#include "DataSourceFilter.h"
+
+#include "EmptySourceFilter.generated.h"
+
+/** Source filter implementation used to replace filter instance who's UClass is not loaded and or does not exist, primarily used by Filter Preset loading */
+UCLASS(NotBlueprintable, hidedropdown)
+class SOURCEFILTERINGTRACE_API  UEmptySourceFilter : public UDataSourceFilter
+{
+	GENERATED_BODY()
+
+public:
+	UPROPERTY()
+	FString MissingClassName;
+
+protected:
+	virtual bool DoesActorPassFilter_Internal(const AActor* InActor) const override
+	{
+		return true;
+	}
+
+	virtual void GetDisplayText_Internal(FText& OutDisplayText) const override
+	{
+		OutDisplayText = FText::Format(NSLOCTEXT("UEmptySourceFilter", "DisplayText", "Empty Filter, missing original Filter Class ({0})"), FText::FromString(MissingClassName));
+	}
+};
\ No newline at end of file
diff --git a/Engine/Plugins/Developer/TraceSourceFiltering/Source/SourceFilteringTrace/Public/SourceFilterCollection.h b/Engine/Plugins/Developer/TraceSourceFiltering/Source/SourceFilteringTrace/Public/SourceFilterCollection.h
new file mode 100644
--- /dev/null
+++ b/Engine/Plugins/Developer/TraceSourceFiltering/Source/SourceFilteringTrace/Public/SourceFilterCollection.h
@@ -0,0 +1,87 @@
+// Copyright Epic Games, Inc. All Rights Reserved.
+
+#pragma once
+
+#include "DataSourceFilter.h"
+#include "DataSourceFilterSet.h"
+#include "UObject/ObjectKey.h"
+#include "Engine/DataAsset.h"
+
+#include "SourceFilterCollection.generated.h"
+
+UCLASS()
+class SOURCEFILTERINGTRACE_API USourceFilterCollection : public UDataAsset
+{
+	GENERATED_BODY()
+public:
+	virtual ~USourceFilterCollection() {}
+
+	/** Begin UDataAsset overrides */
+	virtual void Serialize(FArchive& Ar) override;
+	/** End UDataAsset overrides */
+
+	/** Add Filter instance to the collection, will be added at the root level */
+	void AddFilter(UDataSourceFilter* NewFilter);
+
+	/** Add a Filter instance of the provided class, will be added at the root level */
+	UDataSourceFilter* AddFilterOfClass(const TSubclassOf<UDataSourceFilter>& FilterClass);
+	/** Add a Filter instance of the provided class, added to the provided filter set */
+	UDataSourceFilter* AddFilterOfClassToSet(const TSubclassOf<UDataSourceFilter>& FilterClass, UDataSourceFilterSet* FilterSet);
+
+	/** Remove Filter instance, regardless of whether it is a root-level filter or part of a filter set  */
+	void RemoveFilter(UDataSourceFilter* ToRemoveFilter);
+	/** Remove a Filter Instance from a specific Filter Set */
+	void RemoveFilterFromSet(UDataSourceFilter* ToRemoveFilter, UDataSourceFilterSet* FilterSet);
+	
+	/** Replace a Filter Instance with another */
+	void ReplaceFilter(UDataSourceFilter* Destination, UDataSourceFilter* Source);
+	
+	/** Move a Filter instance to a specific Filter Set (moved to root-level if Destination = nullptr) */
+	void MoveFilter(UDataSourceFilter* Filter, UDataSourceFilterSet* Destination);
+	
+	/** Convert a Filter Instance to a Filter Set (with provided mode), this creates set containing the replace filter */
+	UDataSourceFilterSet* ConvertFilterToSet(UDataSourceFilter* ReplacedFilter, EFilterSetMode Mode);
+	/** Create a Filter set (with provided mode) containing both Filter Instances */
+	UDataSourceFilterSet* MakeFilterSet(UDataSourceFilter* FilterOne, UDataSourceFilter* FilterTwo, EFilterSetMode Mode);
+
+	/** Creates an empty Filter Set (with provided mode) */
+	UDataSourceFilterSet* MakeEmptyFilterSet(EFilterSetMode Mode);
+
+	/** Creates a new collection of filter (sets), provided the filter class names and parent/child relationship */
+	void AddFiltersFromPreset(const TArray<FString>& ClassNames, const TMap<int32, int32>& ChildToParentIndices);
+	
+	/** Resets all contained filter data */
+	void Reset();
+	
+	/** Returns all top-level Filter instances */
+	const TArray<UDataSourceFilter*>& GetFilters() const { return Filters; }
+
+	/** Copies Filter data from other provided Filter Collection*/
+	void CopyData(USourceFilterCollection* OtherCollection);
+protected:
+	/** Recursively removes filter and any contained child filters  */
+	void RemoveFilterRecursive(UDataSourceFilter* ToRemoveFilter);
+
+	/** Adds unique filter class name */
+	void AddClassName(UDataSourceFilter* Filter);
+	void RecursiveRetrieveFilterClassNames(UDataSourceFilter* Filter);
+	UDataSourceFilter* RecursiveCopyFilter(UDataSourceFilter* Filter, int32& FilterOffset);
+	
+	template<typename T> T* CreateNewFilter(UClass* Class = T::StaticClass());
+	void DestroyFilter(UDataSourceFilter* Filter);
+protected:
+	/** Root-level filter instances */
+	UPROPERTY(VisibleAnywhere, Category=Filtering)
+	TArray<UDataSourceFilter*> Filters;
+
+	/** Mapping from Filter Instance FObjectKeys to their class names */
+	TMap<FObjectKey, FString> FilterClassMap;
+
+	/** Flat version of the Filter classes contained by this collection, stored according to Filters ordering, with child filters inline */
+	UPROPERTY()
+	TArray<FString> FilterClasses;
+
+	/** Child / Parent mapping for Filter (sets) */
+	UPROPERTY()
+	TMap<UDataSourceFilter*, UDataSourceFilterSet*> ChildToParent;
+};
\ No newline at end of file
diff --git a/Engine/Plugins/Developer/TraceSourceFiltering/Source/SourceFilteringTrace/Public/SourceFilterTrace.h b/Engine/Plugins/Developer/TraceSourceFiltering/Source/SourceFilteringTrace/Public/SourceFilterTrace.h
new file mode 100644
--- /dev/null
+++ b/Engine/Plugins/Developer/TraceSourceFiltering/Source/SourceFilteringTrace/Public/SourceFilterTrace.h
@@ -0,0 +1,107 @@
+// Copyright Epic Games, Inc. All Rights Reserved.
+
+#pragma once
+
+#include "Misc/Build.h"
+
+#if !IS_PROGRAM && !UE_BUILD_SHIPPING
+#define SOURCE_FILTER_TRACE_ENABLED 1
+#else
+#define SOURCE_FILTER_TRACE_ENABLED 0
+#endif
+
+#if SOURCE_FILTER_TRACE_ENABLED
+
+#include "CoreTypes.h"
+#include "Trace/Config.h"
+#include "Trace/Trace.h"
+#include "UObject/ObjectKey.h"
+#include "Templates/SubclassOf.h"
+
+#include "ObjectTrace.h"
+#include "DataSourceFiltering.h"
+#include "DataSourceFilter.h"
+#include "DataSourceFilterSet.h"
+
+UE_TRACE_CHANNEL_EXTERN(TraceSourceFiltersChannel)
+
+struct SOURCEFILTERINGTRACE_API FSourceFilterTrace
+{
+	/** Output trace data for a UDataSourceFilter (sub) class */
+	static void OutputClass(const TSubclassOf<UDataSourceFilter> InClass);
+	/** Output trace data for a UDataSourceFilter object instance */
+	static void OutputInstance(const UDataSourceFilter* InFilter);
+	/** Output trace data for a UDataSourceFilterSet object instance*/
+	static void OutputSet(const UDataSourceFilterSet* InFilterSet);
+	/** Output trace data for an operation involving a UDataSourceFilter(Set) instance */
+	static void OutputFilterOperation(const UDataSourceFilter* InFilter, ESourceActorFilterOperation Operation, uint64 Parameter);
+
+	/** Output trace data for a change in the UTraceSourceFilteringSettings for this running instance */
+	static void OutputFilterSettingsValue(const FString& InPropertyName, const uint8 InValue);
+
+	/** Output trace data for a UWorld's filtering related information */
+	static void OutputWorld(const UWorld* InWorld);
+	/** Output trace data for an operation involving a UWorld instance */
+	static void OutputWorldOperation(const UWorld* InWorld, EWorldFilterOperation Operation, uint32 Parameter);	
+
+	/** Tries to retrieve a UClass instance according to its Object Identifier */
+	static UClass* RetrieveClassById(uint64 ClassId);
+	/** Tries to retrieve a UClass instance according to its name */
+	static UClass* RetrieveClassByName(const FString& ClassName);
+
+	/** Tries to retrieve a UDataSourceFilter instance according to its Object Identifier */
+	static UDataSourceFilter* RetrieveFilterbyId(uint64 FilterId);
+	/** Tries to retrieve a UWorld instance according to its Object Identifier */
+	static UWorld* RetrieveWorldById(uint64 WorldId);
+
+protected:
+	/** Mapping from a UClass's ObjectKey to their Object identifier retrieved from FObjectTrace::GetObjectId */
+	static TMap<FObjectKey, uint64> FilterClassIds;
+	/** Mapping from a UClass's Object identifier to an FObjectKey */
+	static TMap<uint64, FObjectKey> IDToFilterClass;
+
+	/** Mapping from a UDataSourceFilter's Object identifier to their FObjectKey */
+	static TMap<uint64, FObjectKey> IDToFilter;
+
+	/** Set of FObjectKey's, representing UDataSourceFilter instances, which have previously been traced out */
+	static TSet<FObjectKey> FilterInstances;
+
+	/** Mapping from an UClass's (sub class of UDataSourceFilter) name to their FObjectKey */
+	static TMap<FString, FObjectKey> DataSourceFilterClasses;
+	
+	/** Mapping from an Object identifier (generated from a UWorld instance) to its FObjectKey */
+	static TMap<uint64, FObjectKey> IDsToWorldInstance;
+};
+
+#define TRACE_FILTER_CLASS(Class) \
+	FSourceFilterTrace::OutputClass(Class);
+
+#define TRACE_FILTER_INSTANCE(Instance) \
+	FSourceFilterTrace::OutputInstance(Instance);
+
+#define TRACE_FILTER_SET(Set) \
+	FSourceFilterTrace::OutputSet(Set);
+
+#define TRACE_FILTER_OPERATION(Instance, Operation, Parameter) \
+	FSourceFilterTrace::OutputFilterOperation(Instance, Operation, Parameter);
+
+#define TRACE_FILTER_SETTINGS_VALUE(Name, Value) \
+	FSourceFilterTrace::OutputFilterSettingsValue(Name, Value);
+
+#define TRACE_WORLD_INSTANCE(World) \
+	FSourceFilterTrace::OutputWorld(World);
+
+#define TRACE_WORLD_OPERATION(Instance, Operation, Parameter) \
+	FSourceFilterTrace::OutputWorldOperation(Instance, Operation, Parameter);
+#else
+
+#define TRACE_FILTER_CLASS(Class)
+#define TRACE_FILTER_INSTANCE(Filter)
+#define TRACE_FILTER_SET(Set)
+#define TRACE_FILTER_OPERATION(Instance, Operation, Parameter)
+#define TRACE_FILTER_SETTINGS_VALUE(Name, Value)
+
+#define TRACE_WORLD_INSTANCE(World)
+#define TRACE_WORLD_OPERATION(Instance, Operation, Parameter)
+
+#endif
\ No newline at end of file
diff --git a/Engine/Plugins/Developer/TraceSourceFiltering/Source/SourceFilteringTrace/Public/TraceSourceFiltering.h b/Engine/Plugins/Developer/TraceSourceFiltering/Source/SourceFilteringTrace/Public/TraceSourceFiltering.h
new file mode 100644
--- /dev/null
+++ b/Engine/Plugins/Developer/TraceSourceFiltering/Source/SourceFilteringTrace/Public/TraceSourceFiltering.h
@@ -0,0 +1,48 @@
+// Copyright Epic Games, Inc. All Rights Reserved.
+
+#pragma once
+
+#include "CoreMinimal.h"
+#include "UObject/GCObject.h"
+
+#include "DataSourceFilter.h"
+
+class UTraceSourceFilteringSettings;
+class USourceFilterCollection;
+
+/** Object managing the currently active UDataSourceFilter instances and UTraceSourceFilteringSettings */
+class SOURCEFILTERINGTRACE_API FTraceSourceFiltering : public FGCObject
+{
+public:
+	static void Initialize();
+	static FTraceSourceFiltering& Get();
+
+	/** Returns the running instance its Filter Collection, containing active set of filters */
+	USourceFilterCollection* GetFilterCollection();
+	/** Returns the running instance its Filtering Settings */
+	UTraceSourceFilteringSettings* GetSettings();
+
+	/** Processes an received filtering command, altering the Filter Collection and or Settings accordingly */
+	void ProcessRemoteCommand(const FString& Command, const TArray<FString>& Arguments);
+
+	/** Begin FGCObject overrides */
+	virtual void AddReferencedObjects(FReferenceCollector& Collector) override;
+	/** End FGCObject overrides */
+protected:
+	FTraceSourceFiltering();
+	void PopulateRemoteTraceCommands();
+
+protected:
+	UTraceSourceFilteringSettings* Settings;
+	USourceFilterCollection* FilterCollection;
+
+	/** Structure representing a remotely 'callable' filter command */
+	struct FFilterCommand
+	{
+		TFunction<void(const TArray<FString>&)> Function;
+		int32 NumExpectedArguments;
+	};
+
+	/** Mapping for all filtering commands from their name to respective FFilterCommand object */
+	TMap<FString, FFilterCommand> CommandMap;
+};
\ No newline at end of file
diff --git a/Engine/Plugins/Developer/TraceSourceFiltering/Source/SourceFilteringTrace/Public/TraceWorldFiltering.h b/Engine/Plugins/Developer/TraceSourceFiltering/Source/SourceFilteringTrace/Public/TraceWorldFiltering.h
new file mode 100644
--- /dev/null
+++ b/Engine/Plugins/Developer/TraceSourceFiltering/Source/SourceFilteringTrace/Public/TraceWorldFiltering.h
@@ -0,0 +1,71 @@
+// Copyright Epic Games, Inc. All Rights Reserved.
+
+#pragma once
+
+#include "Delegates/DelegateCombinations.h"
+#include "Engine/World.h"
+
+class FSourceFilterManager;
+
+/** UWorld specific Trace filter, marks individual instances to not be traced out (and all containing actors / objects) */
+struct SOURCEFILTERINGTRACE_API FTraceWorldFiltering
+{
+	DECLARE_MULTICAST_DELEGATE(FTraceWorldFilterStateChanged);
+
+	static void Initialize();
+	static void Destroy();
+	
+	/** Retrieve a FSourceFilterManager instance representing the source filtering for a specific, provided) World instance */
+	static const FSourceFilterManager* GetWorldSourceFilterManager(const UWorld* World);
+
+	/** Return all Worlds currently tracked for Filtering */
+	static const TArray<const UWorld*>& GetWorlds();
+
+	/** Check whether or not a specific World Type can output Trace Data (not filtered out) */
+	static bool IsWorldTypeTraceable(EWorldType::Type InType);
+	/** Check whether or not a specific World's Net Mode can output Trace Data (not filtered out) */
+	static bool IsWorldNetModeTraceable(ENetMode InNetMode);
+	
+	/** Set whether or not a specific World Type should be filtered out (or in) */
+	static void SetStateByWorldType(EWorldType::Type WorldType, bool bState);	
+	/** Set whether or not a specific World Net Mode should be filtered out (or in) */
+	static void SetStateByWorldNetMode(ENetMode NetMode, bool bState);
+	/** Set whether or not a specific UWorld instance's filtering state */
+	static void SetWorldState(const UWorld* InWorld, bool bState);
+
+	/** Returns a user facing display string for the provided UWorld instance */
+	static void GetWorldDisplayString(const UWorld* InWorld, FString& OutDisplayString);
+
+	/** Delegate which will be broadcast whenever the filtering state for any world (type, netmode) changes */
+	static FTraceWorldFilterStateChanged& OnFilterStateChanged();
+
+protected:
+	/** Callbacks used to keep tracking of active (alive) UWorld instances */
+	static void OnWorldInit(UWorld* World, const UWorld::InitializationValues IVS);
+	static void OnWorldPostInit(UWorld* World, const UWorld::InitializationValues IVS);
+	static void OnWorldCleanup(UWorld* World, bool bSessionEnded, bool bCleanupResources);
+	static void RemoveWorld(UWorld* InWorld);
+
+	static void UpdateWorldFiltering();
+
+protected:
+	static FDelegateHandle WorldInitHandle;
+	static FDelegateHandle WorldPostInitHandle;
+	static FDelegateHandle WorldBeginTearDownHandle;
+	static FDelegateHandle WorldCleanupHandle;
+	static FDelegateHandle PreWorldFinishDestroyHandle;
+
+	/** Delegate for broadcasting filtering changes */
+	static FTraceWorldFilterStateChanged FilterStateChangedDelegate;
+
+	/** Array of currently active and alive UWorlds */
+	static TArray<const UWorld*> Worlds;
+	/** Per EWorldType enum entry flag, determines whether or not UWorld's of this type should be filtered out */
+	static TMap<EWorldType::Type, bool> WorldTypeFilterStates;
+	/** Per ENetMode enum entry flag, determines whether or not UWorld's using this netmode should be filtered out */
+	static TMap<ENetMode, bool> NetModeFilterStates;
+	/** Synchronization object for accessing WorldTypeFilterStates and NetModeFilterStates, required to ensure there is not competing access between Networking and Gamethread */
+	static FCriticalSection WorldFilterStatesCritical;
+	/** Mapping from UWorld instance to FSourceFilterManager, entries correspond to world instances in Worlds */
+	static TMap<const UWorld*, FSourceFilterManager*> WorldSourceFilterManagers;
+};
\ No newline at end of file
diff --git a/Engine/Plugins/Developer/TraceSourceFiltering/Source/SourceFilteringTrace/SourceFilteringTrace.Build.cs b/Engine/Plugins/Developer/TraceSourceFiltering/Source/SourceFilteringTrace/SourceFilteringTrace.Build.cs
new file mode 100644
--- /dev/null
+++ b/Engine/Plugins/Developer/TraceSourceFiltering/Source/SourceFilteringTrace/SourceFilteringTrace.Build.cs
@@ -0,0 +1,31 @@
+// Copyright Epic Games, Inc. All Rights Reserved.
+
+using UnrealBuildTool;
+
+public class SourceFilteringTrace : ModuleRules
+{
+	public SourceFilteringTrace(ReadOnlyTargetRules Target) : base(Target)
+	{
+        PrivateDependencyModuleNames.AddRange(
+                new string[] {
+                    "Core",
+					"CoreUObject",
+                    "TraceLog",
+                    "Engine",
+                    "SourceFilteringCore",
+					"PropertyPath",
+					"AssetRegistry",
+					"MessagingCommon"
+                }
+            );
+
+        if (Target.Type == TargetType.Editor)
+        {
+            PrivateDependencyModuleNames.AddRange(
+				new string[] {
+					"UnrealEd"
+				}
+			);
+        }
+    }
+}
diff --git a/Engine/Plugins/Developer/TraceSourceFiltering/TraceSourceFilters.uplugin b/Engine/Plugins/Developer/TraceSourceFiltering/TraceSourceFilters.uplugin
new file mode 100644
--- /dev/null
+++ b/Engine/Plugins/Developer/TraceSourceFiltering/TraceSourceFilters.uplugin
@@ -0,0 +1,46 @@
+{
+	"FileVersion": 1,
+	"Version": 1,
+	"VersionName": "1.0",
+	"FriendlyName": "Insights Data Source Filters",
+	"Description": "Source data filtering for Unreal Insights.",
+	"Category": "Editor",
+	"CreatedBy": "Epic Games, Inc.",
+	"CreatedByURL": "http://epicgames.com",
+	"DocsURL": "",
+	"MarketplaceURL": "",
+	"SupportURL": "",
+	"EnabledByDefault": false,
+	"CanContainContent": true,
+	"IsBetaVersion": false,
+	"Installed": false,
+	"Modules": [
+		{
+			"Name": "SourceFilteringTrace",
+			"Type": "Runtime",
+			"LoadingPhase": "PostEngineInit",
+			"BlacklistTargetConfigurations": [ "Shipping" ]
+		},
+		{
+			"Name": "SourceFilteringCore",
+			"Type": "Runtime",
+			"LoadingPhase": "Default",
+			"BlacklistTargetConfigurations": [ "Shipping" ]
+		},
+		{
+			"Name": "SourceFilteringEditor",
+			"Type": "Editor",
+			"LoadingPhase": "PreDefault"
+		}
+	],
+	"Plugins": [
+		{
+			"Name": "GameplayInsights",
+			"Enabled": true
+		},
+		{
+			"Name": "TcpMessaging",
+			"Enabled": true
+		}
+	]
+}
\ No newline at end of file
