Static analysis fix

diff --git a/Engine/Plugins/Developer/TraceSourceFiltering/Source/SourceFilteringTrace/Private/SourceFilterCollection.cpp b/Engine/Plugins/Developer/TraceSourceFiltering/Source/SourceFilteringTrace/Private/SourceFilterCollection.cpp
--- a/Engine/Plugins/Developer/TraceSourceFiltering/Source/SourceFilteringTrace/Private/SourceFilterCollection.cpp
+++ b/Engine/Plugins/Developer/TraceSourceFiltering/Source/SourceFilteringTrace/Private/SourceFilterCollection.cpp
@@ -38,19 +38,24 @@ UDataSourceFilter* USourceFilterCollection::AddFilterOfClass(const TSubclassOf<U
 
 UDataSourceFilter* USourceFilterCollection::AddFilterOfClassToSet(const TSubclassOf<UDataSourceFilter>& FilterClass, UDataSourceFilterSet* FilterSet)
 {
-	ensure(FilterClass.Get() && FilterSet);
-	FilterSet->Modify();
+	checkf(FilterClass.Get() && FilterSet, TEXT("Cannot add filter using a null class, or null target set"));
+	if (FilterSet)
+	{
+		FilterSet->Modify();
 
-	UDataSourceFilter* NewFilter = CreateNewFilter<UDataSourceFilter>(FilterClass.Get());
-	FilterSet->Filters.Add(NewFilter);
-	ChildToParent.Add(NewFilter, FilterSet);
-	AddClassName(NewFilter);
+		UDataSourceFilter* NewFilter = CreateNewFilter<UDataSourceFilter>(FilterClass.Get());
+		FilterSet->Filters.Add(NewFilter);
+		ChildToParent.Add(NewFilter, FilterSet);
+		AddClassName(NewFilter);
 
-	TRACE_FILTER_INSTANCE(NewFilter);
+		TRACE_FILTER_INSTANCE(NewFilter);
 
-	TRACE_FILTER_OPERATION(NewFilter, ESourceActorFilterOperation::MoveFilter, FObjectTrace::GetObjectId(FilterSet));
+		TRACE_FILTER_OPERATION(NewFilter, ESourceActorFilterOperation::MoveFilter, FObjectTrace::GetObjectId(FilterSet));
 
-	return NewFilter;
+		return NewFilter;
+	}
+
+	return nullptr;
 }
 
 UDataSourceFilterSet* USourceFilterCollection::ConvertFilterToSet(UDataSourceFilter* ReplacedFilter, EFilterSetMode Mode)
