//UE4/Main - Compile UE4Game HoloLens - Multiple identifier not found errors in Atomic.h and SourceFilterCollection.cpp
#fix ensure that trace dependent code paths are guarded against it being turned off

diff --git a/Engine/Plugins/Developer/TraceSourceFiltering/Source/SourceFilteringTrace/Private/SourceFilterCollection.cpp b/Engine/Plugins/Developer/TraceSourceFiltering/Source/SourceFilteringTrace/Private/SourceFilterCollection.cpp
--- a/Engine/Plugins/Developer/TraceSourceFiltering/Source/SourceFilteringTrace/Private/SourceFilterCollection.cpp
+++ b/Engine/Plugins/Developer/TraceSourceFiltering/Source/SourceFilteringTrace/Private/SourceFilterCollection.cpp
@@ -50,7 +50,7 @@ UDataSourceFilter* USourceFilterCollection::AddFilterOfClassToSet(const TSubclas
 
 		TRACE_FILTER_INSTANCE(NewFilter);
 
-		TRACE_FILTER_OPERATION(NewFilter, ESourceActorFilterOperation::MoveFilter, FObjectTrace::GetObjectId(FilterSet));
+		TRACE_FILTER_OPERATION(NewFilter, ESourceActorFilterOperation::MoveFilter, TRACE_FILTER_IDENTIFIER(FilterSet));
 
 		return NewFilter;
 	}
@@ -131,6 +131,7 @@ void USourceFilterCollection::AddFiltersFromPreset(const TArray<FString>& ClassN
 			}
 			else
 			{
+#if SOURCE_FILTER_TRACE_ENABLED
 				// This is a filter, so create it according to the class name				
 				if (UClass* Class = FSourceFilterTrace::RetrieveClassByName(ClassNames[FilterIndex]))
 				{
@@ -140,6 +141,7 @@ void USourceFilterCollection::AddFiltersFromPreset(const TArray<FString>& ClassN
 						AddFilter(Filter);
 					}
 				}
+#endif // SOURCE_FILTER_TRACE_ENABLED
 			}
 
 			// In case we were unable to retrieve a Filter instance its class replace it with an empty 'stub' filter instance, indicating to the user that it, and its class, is missing
@@ -199,7 +201,7 @@ void USourceFilterCollection::RemoveFilter(UDataSourceFilter* ToRemoveFilter)
 		FilterSet->Filters.Empty();
 	}
 
-	TRACE_FILTER_OPERATION(ToRemoveFilter, ESourceActorFilterOperation::RemoveFilter, FObjectTrace::GetObjectId(Cast<UDataSourceFilterSet>(ToRemoveFilter)));
+	TRACE_FILTER_OPERATION(ToRemoveFilter, ESourceActorFilterOperation::RemoveFilter, TRACE_FILTER_IDENTIFIER(Cast<UDataSourceFilterSet>(ToRemoveFilter)));
 
 	if (OuterFilterSet)
 	{
@@ -235,12 +237,12 @@ void USourceFilterCollection::RemoveFilterRecursive(UDataSourceFilter* ToRemoveF
 		}
 	}
 
-	TRACE_FILTER_OPERATION(ToRemoveFilter, ESourceActorFilterOperation::RemoveFilter, FObjectTrace::GetObjectId(Cast<UDataSourceFilterSet>(ToRemoveFilter)));
+	TRACE_FILTER_OPERATION(ToRemoveFilter, ESourceActorFilterOperation::RemoveFilter, TRACE_FILTER_IDENTIFIER(Cast<UDataSourceFilterSet>(ToRemoveFilter)));
 }
 
 void USourceFilterCollection::ReplaceFilter(UDataSourceFilter* Destination, UDataSourceFilter* Source)
 {
-	TRACE_FILTER_OPERATION(Source, ESourceActorFilterOperation::ReplaceFilter, FObjectTrace::GetObjectId(Destination));
+	TRACE_FILTER_OPERATION(Source, ESourceActorFilterOperation::ReplaceFilter, TRACE_FILTER_IDENTIFIER(Destination));
 
 	UDataSourceFilterSet* OuterFilterSet = ChildToParent.FindChecked(Destination);
 	if (OuterFilterSet)
@@ -261,7 +263,7 @@ void USourceFilterCollection::ReplaceFilter(UDataSourceFilter* Destination, UDat
 
 void USourceFilterCollection::MoveFilter(UDataSourceFilter* Filter, UDataSourceFilterSet* Destination)
 {
-	TRACE_FILTER_OPERATION(Filter, ESourceActorFilterOperation::MoveFilter, Destination ? FObjectTrace::GetObjectId(Destination) : 0);
+	TRACE_FILTER_OPERATION(Filter, ESourceActorFilterOperation::MoveFilter, Destination ? TRACE_FILTER_IDENTIFIER(Destination) : 0);
 
 	UDataSourceFilterSet*& FilterParent = ChildToParent.FindChecked(Filter);
 
diff --git a/Engine/Plugins/Developer/TraceSourceFiltering/Source/SourceFilteringTrace/Private/SourceFilterTrace.cpp b/Engine/Plugins/Developer/TraceSourceFiltering/Source/SourceFilteringTrace/Private/SourceFilterTrace.cpp
--- a/Engine/Plugins/Developer/TraceSourceFiltering/Source/SourceFilteringTrace/Private/SourceFilterTrace.cpp
+++ b/Engine/Plugins/Developer/TraceSourceFiltering/Source/SourceFilteringTrace/Private/SourceFilterTrace.cpp
@@ -85,7 +85,7 @@ void FSourceFilterTrace::OutputClass(const TSubclassOf<UDataSourceFilter> InClas
 
 	DataSourceFilterClasses.Add(Class->GetName(), Class);
 
-	const uint64 Identifier = FObjectTrace::GetObjectId(Class);
+	const uint64 Identifier = TRACE_FILTER_IDENTIFIER(Class);
 	if (!FilterClassIds.Contains(Class))
 	{
 		FilterClassIds.Add(Class, Identifier);
@@ -102,8 +102,8 @@ void FSourceFilterTrace::OutputInstance(const UDataSourceFilter* InFilter)
 	{
 		TRACE_FILTER_CLASS(InFilter->GetClass());
 
-		const uint64 ClassId = FObjectTrace::GetObjectId(InFilter->GetClass());
-		const uint64 InstanceId = FObjectTrace::GetObjectId(InFilter);
+		const uint64 ClassId = TRACE_FILTER_IDENTIFIER(InFilter->GetClass());
+		const uint64 InstanceId = TRACE_FILTER_IDENTIFIER(InFilter);
 		const uint64 SetId = 0;
 
 		FilterInstances.Add(InFilter);
@@ -134,7 +134,7 @@ void FSourceFilterTrace::OutputSet(const UDataSourceFilterSet* InFilterSet)
 {
 	if (InFilterSet)
 	{
-		const uint64 SetId = FObjectTrace::GetObjectId(InFilterSet);
+		const uint64 SetId = TRACE_FILTER_IDENTIFIER(InFilterSet);
 
 		IDToFilter.Add(SetId, InFilterSet);
 
@@ -146,7 +146,7 @@ void FSourceFilterTrace::OutputSet(const UDataSourceFilterSet* InFilterSet)
 
 void FSourceFilterTrace::OutputFilterOperation(const UDataSourceFilter* InFilter, ESourceActorFilterOperation Operation, uint64 Parameter)
 {
-	const uint64 InstanceId = FObjectTrace::GetObjectId(InFilter);
+	const uint64 InstanceId = TRACE_FILTER_IDENTIFIER(InFilter);
 	UE_TRACE_LOG(SourceFilters, FilterOperation, TraceSourceFiltersChannel)
 		<< FilterOperation.InstanceId(InstanceId)
 		<< FilterOperation.Operation((uint8)Operation)
@@ -170,7 +170,7 @@ void FSourceFilterTrace::OutputFilterSettingsValue(const FString& InPropertyName
 
 void FSourceFilterTrace::OutputWorld(const UWorld* InWorld)
 {
-	const uint64 InstanceId = FObjectTrace::GetObjectId(InWorld);
+	const uint64 InstanceId = TRACE_FILTER_IDENTIFIER(InWorld);
 	
 	if (!IDsToWorldInstance.Contains(InstanceId))
 	{
@@ -197,7 +197,7 @@ void FSourceFilterTrace::OutputWorld(const UWorld* InWorld)
 
 void FSourceFilterTrace::OutputWorldOperation(const UWorld* InWorld, EWorldFilterOperation Operation, uint32 Parameter)
 {
-	const uint64 InstanceId = FObjectTrace::GetObjectId(InWorld);
+	const uint64 InstanceId = TRACE_FILTER_IDENTIFIER(InWorld);
 	UE_TRACE_LOG(WorldSourceFilters, WorldOperation, TraceSourceFiltersChannel)
 		<< WorldOperation.InstanceId(InstanceId)
 		<< WorldOperation.Operation((uint8)Operation)
diff --git a/Engine/Plugins/Developer/TraceSourceFiltering/Source/SourceFilteringTrace/Private/SourceFilteringTraceModule.cpp b/Engine/Plugins/Developer/TraceSourceFiltering/Source/SourceFilteringTrace/Private/SourceFilteringTraceModule.cpp
--- a/Engine/Plugins/Developer/TraceSourceFiltering/Source/SourceFilteringTrace/Private/SourceFilteringTraceModule.cpp
+++ b/Engine/Plugins/Developer/TraceSourceFiltering/Source/SourceFilteringTrace/Private/SourceFilteringTraceModule.cpp
@@ -74,8 +74,10 @@ void FSourceFilteringTraceModule::StartupModule()
 	FTraceWorldFiltering::Initialize();
 	FTraceSourceFiltering::Initialize();
 
+#if SOURCE_FILTER_TRACE_ENABLED
 	// Forcefully enable the source trace channel
 	Trace::FChannel::Toggle(&TraceSourceFiltersChannel, true);
+#endif // SOURCE_FILTER_TRACE_ENABLED
 
 #if WITH_EDITOR
 	// Add callback to trace out Filter Classes once the Asset Registry has finished loading 
diff --git a/Engine/Plugins/Developer/TraceSourceFiltering/Source/SourceFilteringTrace/Private/TraceSourceFiltering.cpp b/Engine/Plugins/Developer/TraceSourceFiltering/Source/SourceFilteringTrace/Private/TraceSourceFiltering.cpp
--- a/Engine/Plugins/Developer/TraceSourceFiltering/Source/SourceFilteringTrace/Private/TraceSourceFiltering.cpp
+++ b/Engine/Plugins/Developer/TraceSourceFiltering/Source/SourceFilteringTrace/Private/TraceSourceFiltering.cpp
@@ -36,6 +36,7 @@ UTraceSourceFilteringSettings* FTraceSourceFiltering::GetSettings()
 	return Settings;
 }
 
+#if SOURCE_FILTER_TRACE_ENABLED
 static UClass* ConvertArgumentToClass(FString Argument)
 {
 	const uint64 ClassId = FCString::Atoi64(*Argument);
@@ -79,6 +80,7 @@ static UWorld* ConvertArgumentToWorld(FString Argument)
 
 	return World;
 }
+#endif // SOURCE_FILTER_TRACE_ENABLED
 
 void FTraceSourceFiltering::ProcessRemoteCommand(const FString& Command, const TArray<FString>& Arguments)
 {
@@ -97,6 +99,7 @@ void FTraceSourceFiltering::AddReferencedObjects(FReferenceCollector& Collector)
 
 void FTraceSourceFiltering::PopulateRemoteTraceCommands()
 {
+#if SOURCE_FILTER_TRACE_ENABLED
 	CommandMap.Add(TEXT("AddFilterById"), { [this](const TArray<FString>& Arguments) -> void { FilterCollection->AddFilterOfClass(ConvertArgumentToClass(Arguments[0])); }, 1 });
 
 	CommandMap.Add(TEXT("AddFilterClassToSet"),
@@ -219,4 +222,5 @@ void FTraceSourceFiltering::PopulateRemoteTraceCommands()
 			2
 		}
 	);
+#endif // SOURCE_FILTER_TRACE_ENABLED
 }
diff --git a/Engine/Plugins/Developer/TraceSourceFiltering/Source/SourceFilteringTrace/Public/SourceFilterTrace.h b/Engine/Plugins/Developer/TraceSourceFiltering/Source/SourceFilteringTrace/Public/SourceFilterTrace.h
--- a/Engine/Plugins/Developer/TraceSourceFiltering/Source/SourceFilteringTrace/Public/SourceFilterTrace.h
+++ b/Engine/Plugins/Developer/TraceSourceFiltering/Source/SourceFilteringTrace/Public/SourceFilterTrace.h
@@ -4,7 +4,7 @@
 
 #include "Misc/Build.h"
 
-#if !IS_PROGRAM && !UE_BUILD_SHIPPING
+#if UE_TRACE_ENABLED && !IS_PROGRAM && !UE_BUILD_SHIPPING
 #define SOURCE_FILTER_TRACE_ENABLED 1
 #else
 #define SOURCE_FILTER_TRACE_ENABLED 0
@@ -93,6 +93,9 @@ protected:
 
 #define TRACE_WORLD_OPERATION(Instance, Operation, Parameter) \
 	FSourceFilterTrace::OutputWorldOperation(Instance, Operation, Parameter);
+
+#define TRACE_FILTER_IDENTIFIER(Object) \
+	FObjectTrace::GetObjectId(Object)
 #else
 
 #define TRACE_FILTER_CLASS(Class)
@@ -104,4 +107,6 @@ protected:
 #define TRACE_WORLD_INSTANCE(World)
 #define TRACE_WORLD_OPERATION(Instance, Operation, Parameter)
 
+#define TRACE_FILTER_IDENTIFIER(Object) 0
+
 #endif
\ No newline at end of file
