Add ability to create a data source filter class from within Source Filter tab
#fix add menu builder entry to filter selection menu, allows for creating a new blueprint class + add it to the filter set

diff --git a/Engine/Plugins/Developer/TraceSourceFiltering/Source/SourceFilteringEditor/Private/EditorSessionSourceFilterService.cpp b/Engine/Plugins/Developer/TraceSourceFiltering/Source/SourceFilteringEditor/Private/EditorSessionSourceFilterService.cpp
--- a/Engine/Plugins/Developer/TraceSourceFiltering/Source/SourceFilteringEditor/Private/EditorSessionSourceFilterService.cpp
+++ b/Engine/Plugins/Developer/TraceSourceFiltering/Source/SourceFilteringEditor/Private/EditorSessionSourceFilterService.cpp
@@ -8,6 +8,7 @@
 #include "UObject/SoftObjectPath.h"
 #include "ScopedTransaction.h"
 #include "Editor.h"
+#include "Factories/BlueprintFactory.h"
 
 #include "TraceSourceFiltering.h"
 #include "IFilterObject.h"
@@ -32,6 +33,8 @@
 #include "SourceFilterStyle.h"
 #include "TraceFilter.h"
 #include "TraceSourceFilteringProjectSettings.h"
+#include "Editor/EditorEngine.h"
+#include "Subsystems/AssetEditorSubsystem.h"
 
 #define LOCTEXT_NAMESPACE "FEditorSourceFilterService"
 
@@ -210,7 +213,7 @@ TSharedRef<SWidget> FEditorSessionSourceFilterService::GetFilterPickerWidget(FOn
 			[
 				FModuleManager::LoadModuleChecked<FClassViewerModule>("ClassViewer").CreateClassViewer(Options, ClassPicked)
 			]
-		];	
+		];
 }
 
 /** Helper function to have user pick a (existing) package for saving a filter preset */
@@ -358,6 +361,56 @@ TSharedPtr<FExtender> FEditorSessionSourceFilterService::GetExtender()
 				);
 			})
 		);
+
+		Extender->AddMenuExtension(FName("FilterPicker"), EExtensionHook::After, TSharedPtr<FUICommandList>(), FMenuExtensionDelegate::CreateLambda([this](FMenuBuilder& MenuBuilder) -> void
+		{
+			MenuBuilder.AddMenuEntry(LOCTEXT("NewSourceFilterBPLabel", "Create new Filter Blueprint"),
+					LOCTEXT("NewSourceFilterBPLabel", "Create new Filter Blueprint"),
+					FSlateIcon(),
+					FUIAction(FExecuteAction::CreateLambda([this]()
+					{
+						FString AssetPath;
+						const FString DefaultFilesystemDirectory = FEditorDirectories::Get().GetLastDirectory(ELastDirectory::NEW_ASSET);
+						if (DefaultFilesystemDirectory.IsEmpty() || !FPackageName::TryConvertFilenameToLongPackageName(DefaultFilesystemDirectory, AssetPath))
+						{
+							// No saved path, just use the game content root
+							AssetPath = TEXT("/Game");
+						}
+						
+						// Let user determine new path for Blueprint asset
+						FSaveAssetDialogConfig SaveAssetDialogConfig;
+						SaveAssetDialogConfig.DialogTitleOverride = LOCTEXT("SaveAssetDialogTitle", "Save Asset As");
+						SaveAssetDialogConfig.DefaultPath = AssetPath;
+						SaveAssetDialogConfig.ExistingAssetPolicy = ESaveAssetDialogExistingAssetPolicy::AllowButWarn;
+
+						FContentBrowserModule& ContentBrowserModule = FModuleManager::LoadModuleChecked<FContentBrowserModule>("ContentBrowser");
+						FString SaveObjectPath = ContentBrowserModule.Get().CreateModalSaveAssetDialog(SaveAssetDialogConfig);
+
+						// Check that we have a valid object path to create the blueprint
+						if (!SaveObjectPath.IsEmpty())
+						{
+							const FString SavePackageName = FPackageName::ObjectPathToPackageName(SaveObjectPath);
+							const FString SavePackagePath = FPaths::GetPath(SavePackageName);
+							const FString SaveAssetName = FPaths::GetBaseFilename(SavePackageName);
+							FEditorDirectories::Get().SetLastDirectory(ELastDirectory::NEW_ASSET, SavePackagePath);
+
+							IAssetTools& AssetTools = FModuleManager::GetModuleChecked<FAssetToolsModule>("AssetTools").Get();
+							UBlueprintFactory* Factory = NewObject<UBlueprintFactory>();
+							Factory->ParentClass = UDataSourceFilter::StaticClass();
+
+							// Generate blueprint with DataSourceFilter as its parent class
+							if (UBlueprint* NewFilterBlueprint = Cast<UBlueprint>(AssetTools.CreateAsset(SaveAssetName, SavePackagePath, UBlueprint::StaticClass(), Factory)))
+							{
+								if (GEditor)
+								{
+									GEditor->GetEditorSubsystem<UAssetEditorSubsystem>()->OpenEditorForAsset(NewFilterBlueprint);
+									AddFilter(NewFilterBlueprint->GeneratedClass.Get()->GetPathName());
+								}
+							}
+						}
+					}))
+				);
+		}));
 	}
 
 	return Extender;
diff --git a/Engine/Plugins/Developer/TraceSourceFiltering/Source/SourceFilteringEditor/Private/STraceSourceFilteringWidget.cpp b/Engine/Plugins/Developer/TraceSourceFiltering/Source/SourceFilteringEditor/Private/STraceSourceFilteringWidget.cpp
--- a/Engine/Plugins/Developer/TraceSourceFiltering/Source/SourceFilteringEditor/Private/STraceSourceFilteringWidget.cpp
+++ b/Engine/Plugins/Developer/TraceSourceFiltering/Source/SourceFilteringEditor/Private/STraceSourceFilteringWidget.cpp
@@ -167,7 +167,14 @@ void STraceSourceFilteringWidget::ConstructMenuBox()
 		.OnGetMenuContent(FOnGetContent::CreateLambda([OnFilterClassPicked, this]()
 		{
 			FMenuBuilder MenuBuilder(true, TSharedPtr<FUICommandList>(), SessionFilterService->GetExtender());
-			MenuBuilder.AddWidget(SessionFilterService->GetFilterPickerWidget(FOnFilterClassPicked::CreateLambda(OnFilterClassPicked)), FText::GetEmpty(), true, false);
+			
+			MenuBuilder.BeginSection(FName("FilterPicker"));
+			{
+				MenuBuilder.AddWidget(SessionFilterService->GetFilterPickerWidget(FOnFilterClassPicked::CreateLambda(OnFilterClassPicked)), FText::GetEmpty(), true, false);
+			}
+			MenuBuilder.EndSection();
+			
+
 			return MenuBuilder.MakeWidget();
 		}))
 		.ButtonContent()
