Data Source Filters are not correctly applied to Actors
#fix Ensure that we update world filtering, after a UWorld has finished Actor initialization, to handle NetMode changes

diff --git a/Engine/Plugins/Developer/TraceSourceFiltering/Source/SourceFilteringTrace/Private/TraceWorldFiltering.cpp b/Engine/Plugins/Developer/TraceSourceFiltering/Source/SourceFilteringTrace/Private/TraceWorldFiltering.cpp
--- a/Engine/Plugins/Developer/TraceSourceFiltering/Source/SourceFilteringTrace/Private/TraceWorldFiltering.cpp
+++ b/Engine/Plugins/Developer/TraceSourceFiltering/Source/SourceFilteringTrace/Private/TraceWorldFiltering.cpp
@@ -206,6 +206,8 @@ void FTraceWorldFiltering::OnWorldInit(UWorld* InWorld, const UWorld::Initializa
 
 void FTraceWorldFiltering::OnWorldPostInit(UWorld* World, const UWorld::InitializationValues IVS)
 {
+	UpdateWorldFiltering();
+
 	// Handle all loaded AActors within this world, to handle duplicated UWorlds (PIE)
 	FSourceFilterManager* Manager = WorldSourceFilterManagers.FindChecked(World);
 	Manager->Tick(0.f);
