- Fixed a crash in TraceLog with FSourceFilterTrace::OutputWorld
#rnx

diff --git a/Engine/Plugins/Developer/TraceSourceFiltering/Source/SourceFilteringTrace/Private/SourceFilterTrace.cpp b/Engine/Plugins/Developer/TraceSourceFiltering/Source/SourceFilteringTrace/Private/SourceFilterTrace.cpp
--- a/Engine/Plugins/Developer/TraceSourceFiltering/Source/SourceFilteringTrace/Private/SourceFilterTrace.cpp
+++ b/Engine/Plugins/Developer/TraceSourceFiltering/Source/SourceFilteringTrace/Private/SourceFilterTrace.cpp
@@ -184,7 +184,7 @@ void FSourceFilterTrace::OutputWorld(const UWorld* InWorld)
 		{
 			TCHAR* StringPtr = (TCHAR*)Out;
 			FMemory::Memcpy(StringPtr, *SendName, AttachmentSize);
-			StringPtr[AttachmentSize - 1] = '\0';
+			StringPtr[AttachmentSize / sizeof(TCHAR) - 1] = '\0';
 		};
 
 		UE_TRACE_LOG(WorldSourceFilters, WorldInstance, TraceSourceFiltersChannel, AttachmentSize)
