- More TraceLog attachment fixes
#rnx

diff --git a/Engine/Plugins/Developer/TraceSourceFiltering/Source/SourceFilteringTrace/Private/SourceFilterTrace.cpp b/Engine/Plugins/Developer/TraceSourceFiltering/Source/SourceFilteringTrace/Private/SourceFilterTrace.cpp
--- a/Engine/Plugins/Developer/TraceSourceFiltering/Source/SourceFilteringTrace/Private/SourceFilterTrace.cpp
+++ b/Engine/Plugins/Developer/TraceSourceFiltering/Source/SourceFilteringTrace/Private/SourceFilterTrace.cpp
@@ -119,7 +119,7 @@ void FSourceFilterTrace::OutputInstance(const UDataSourceFilter* InFilter)
 		{
 			TCHAR* StringPtr = (TCHAR*)Out;
 			FMemory::Memcpy(StringPtr, *DisplayString, AttachmentSize);
-			StringPtr[AttachmentSize - 1] = '\0';
+			StringPtr[AttachmentSize / sizeof(TCHAR) - 1] = '\0';
 		};
 		
 		UE_TRACE_LOG(SourceFilters, FilterInstance, TraceSourceFiltersChannel, AttachmentSize)
@@ -160,7 +160,7 @@ void FSourceFilterTrace::OutputFilterSettingsValue(const FString& InPropertyName
 	{
 		TCHAR* StringPtr = (TCHAR*)Out;
 		FMemory::Memcpy(StringPtr, *InPropertyName, AttachmentSize);
-		StringPtr[AttachmentSize - 1] = '\0';
+		StringPtr[AttachmentSize / sizeof(TCHAR) - 1] = '\0';
 	};
 
 	UE_TRACE_LOG(SourceFilters, SetFilterSettingValue, TraceSourceFiltersChannel, AttachmentSize)
