Remove raw pointers from TraceSourceFilteringPlugin

diff --git a/Engine/Plugins/Developer/TraceSourceFiltering/Source/SourceFilteringEditor/Private/SourceFilteringEditorModule.cpp b/Engine/Plugins/Developer/TraceSourceFiltering/Source/SourceFilteringEditor/Private/SourceFilteringEditorModule.cpp
--- a/Engine/Plugins/Developer/TraceSourceFiltering/Source/SourceFilteringEditor/Private/SourceFilteringEditorModule.cpp
+++ b/Engine/Plugins/Developer/TraceSourceFiltering/Source/SourceFilteringEditor/Private/SourceFilteringEditorModule.cpp
@@ -101,7 +101,7 @@ void FSourceFilteringEditorModule::OnAssetsPendingDelete(TArray<UObject*> const&
 
 	if (FilterCollection)
 	{
-		TArray<UDataSourceFilter*> Filters;
+		TArray<TObjectPtr<UDataSourceFilter>> Filters;
 		FilterCollection->GetFlatFilters(Filters);
 		for (UDataSourceFilter* Filter : Filters) 
 		{
diff --git a/Engine/Plugins/Developer/TraceSourceFiltering/Source/SourceFilteringTrace/Private/SourceFilterCollection.cpp b/Engine/Plugins/Developer/TraceSourceFiltering/Source/SourceFilteringTrace/Private/SourceFilterCollection.cpp
--- a/Engine/Plugins/Developer/TraceSourceFiltering/Source/SourceFilteringTrace/Private/SourceFilterCollection.cpp
+++ b/Engine/Plugins/Developer/TraceSourceFiltering/Source/SourceFilteringTrace/Private/SourceFilterCollection.cpp
@@ -27,14 +27,14 @@ void USourceFilterCollection::OnObjectsReplaced(const TMap<UObject*, UObject*>&
 		}
 	}
 
-	TArray<UDataSourceFilter*> KeyFilters;
+	TArray<TObjectPtr<UDataSourceFilter>> KeyFilters;
 	ChildToParent.GenerateKeyArray(KeyFilters);
 
 	for (UDataSourceFilter* Filter : KeyFilters)
 	{
 		if (UObject* const * NewObjectPtr = OldToNewInstanceMap.Find(Filter))
 		{
-			UDataSourceFilterSet* ParentFilterSet = nullptr;			
+			TObjectPtr<UDataSourceFilterSet> ParentFilterSet = nullptr;			
 			ChildToParent.RemoveAndCopyValue(Filter, ParentFilterSet);
 
 			ChildToParent.Add(Cast<UDataSourceFilter>(*NewObjectPtr), ParentFilterSet);
@@ -268,7 +268,7 @@ void USourceFilterCollection::AddFiltersFromPreset(const TArray<FString>& ClassN
 
 void USourceFilterCollection::RemoveFilter(UDataSourceFilter* ToRemoveFilter)
 {
-	UDataSourceFilterSet* OuterFilterSet = nullptr;
+	TObjectPtr<UDataSourceFilterSet> OuterFilterSet = nullptr;
 	ChildToParent.RemoveAndCopyValue(ToRemoveFilter, OuterFilterSet);
 	
 	// In case of a set, also remove children 
@@ -346,7 +346,7 @@ void USourceFilterCollection::ReplaceFilter(UDataSourceFilter* Destination, UDat
 {
 	TRACE_FILTER_OPERATION(Source, ESourceActorFilterOperation::ReplaceFilter, TRACE_FILTER_IDENTIFIER(Destination));
 
-	UDataSourceFilterSet* OuterFilterSet = nullptr;
+	TObjectPtr<UDataSourceFilterSet> OuterFilterSet = nullptr;
 	ChildToParent.RemoveAndCopyValue(Destination, OuterFilterSet);
 
 	if (OuterFilterSet)
@@ -371,7 +371,7 @@ void USourceFilterCollection::MoveFilter(UDataSourceFilter* Filter, UDataSourceF
 {
 	TRACE_FILTER_OPERATION(Filter, ESourceActorFilterOperation::MoveFilter, Destination ? TRACE_FILTER_IDENTIFIER(Destination) : 0);
 
-	UDataSourceFilterSet*& FilterParent = ChildToParent.FindChecked(Filter);
+	TObjectPtr<UDataSourceFilterSet>& FilterParent = ChildToParent.FindChecked(Filter);
 
 	if (FilterParent)
 	{
@@ -423,7 +423,7 @@ void USourceFilterCollection::Reset()
 	SourceFiltersUpdatedDelegate.Broadcast();
 }
 
-void USourceFilterCollection::GetFlatFilters(TArray<UDataSourceFilter*>& OutFilters)
+void USourceFilterCollection::GetFlatFilters(TArray<TObjectPtr<UDataSourceFilter>>& OutFilters)
 {
 	ChildToParent.GenerateKeyArray(OutFilters);
 }
diff --git a/Engine/Plugins/Developer/TraceSourceFiltering/Source/SourceFilteringTrace/Private/TraceSourceFiltering.cpp b/Engine/Plugins/Developer/TraceSourceFiltering/Source/SourceFilteringTrace/Private/TraceSourceFiltering.cpp
--- a/Engine/Plugins/Developer/TraceSourceFiltering/Source/SourceFilteringTrace/Private/TraceSourceFiltering.cpp
+++ b/Engine/Plugins/Developer/TraceSourceFiltering/Source/SourceFilteringTrace/Private/TraceSourceFiltering.cpp
@@ -7,11 +7,13 @@
 #include "SourceFilterTrace.h"
 #include "TraceWorldFiltering.h"
 #include "Misc/CommandLine.h"
+#include "HAL/LowLevelMemTracker.h"
 
 #include "TraceSourceFilteringProjectSettings.h"
 
 FTraceSourceFiltering::FTraceSourceFiltering()
 {
+	LLM_SCOPE_BYNAME("Debug/Profiling");
 	FilterCollection = NewObject<USourceFilterCollection>(GetTransientPackage(), NAME_None, RF_Transactional | RF_Transient);
 	Settings = DuplicateObject(GetDefault<UTraceSourceFilteringSettings>(), nullptr);
 
@@ -112,12 +114,14 @@ void FTraceSourceFiltering::ProcessRemoteCommand(const FString& Command, const T
 
 void FTraceSourceFiltering::AddReferencedObjects(FReferenceCollector& Collector)
 {
+	LLM_SCOPE_BYNAME("Debug/Profiling");
 	Collector.AddReferencedObject(FilterCollection);
 	Collector.AddReferencedObject(Settings);
 }
 
 void FTraceSourceFiltering::PopulateRemoteTraceCommands()
 {
+	LLM_SCOPE_BYNAME("Debug/Profiling");
 #if SOURCE_FILTER_TRACE_ENABLED
 CommandMap.Add(TEXT("AddFilterById"), { [this](const TArray<FString>& Arguments) -> void { FilterCollection->AddFilterOfClass(ConvertArgumentToClass(Arguments[0])); }, 1 });
 
diff --git a/Engine/Plugins/Developer/TraceSourceFiltering/Source/SourceFilteringTrace/Public/DataSourceFilterSet.h b/Engine/Plugins/Developer/TraceSourceFiltering/Source/SourceFilteringTrace/Public/DataSourceFilterSet.h
--- a/Engine/Plugins/Developer/TraceSourceFiltering/Source/SourceFilteringTrace/Public/DataSourceFilterSet.h
+++ b/Engine/Plugins/Developer/TraceSourceFiltering/Source/SourceFilteringTrace/Public/DataSourceFilterSet.h
@@ -17,7 +17,7 @@ class SOURCEFILTERINGTRACE_API UDataSourceFilterSet : public UDataSourceFilter,
 
 	GENERATED_BODY()
 public:	
-	const TArray<UDataSourceFilter*>& GetFilters() const { return Filters; }
+	const TArray<TObjectPtr<UDataSourceFilter>>& GetFilters() const { return Filters; }
 	void SetFilterMode(EFilterSetMode InMode);
 
 	/** Begin IDataSourceFilterSetInterface overrides */
@@ -34,7 +34,7 @@ protected:
 protected:
 	/** Contained Filter instance */
 	UPROPERTY()
-	TArray<UDataSourceFilter*> Filters;
+	TArray<TObjectPtr<UDataSourceFilter>> Filters;
 
 	/** Current Filter set operation */
 	UPROPERTY()
diff --git a/Engine/Plugins/Developer/TraceSourceFiltering/Source/SourceFilteringTrace/Public/SourceFilterCollection.h b/Engine/Plugins/Developer/TraceSourceFiltering/Source/SourceFilteringTrace/Public/SourceFilterCollection.h
--- a/Engine/Plugins/Developer/TraceSourceFiltering/Source/SourceFilteringTrace/Public/SourceFilterCollection.h
+++ b/Engine/Plugins/Developer/TraceSourceFiltering/Source/SourceFilteringTrace/Public/SourceFilterCollection.h
@@ -66,7 +66,7 @@ public:
 	const TArray<UDataSourceFilter*>& GetFilters() const { return Filters; }
 
 	/** Returns flattened Filter instances */
-	void GetFlatFilters(TArray<UDataSourceFilter*>& OutFilters);
+	void GetFlatFilters(TArray<TObjectPtr<UDataSourceFilter>>& OutFilters);
 
 	/** Copies Filter data from other provided Filter Collection*/
 	void CopyData(USourceFilterCollection* OtherCollection);
@@ -101,7 +101,7 @@ protected:
 protected:
 	/** Root-level filter instances */
 	UPROPERTY(VisibleAnywhere, Category=Filtering)
-	TArray<UDataSourceFilter*> Filters;
+	TArray<TObjectPtr<UDataSourceFilter>> Filters;
 
 	/** Class filters, used for high-level filtering of AActor instances inside of a UWorld */
 	UPROPERTY(VisibleAnywhere, Category = Filtering)
@@ -116,7 +116,7 @@ protected:
 
 	/** Child / Parent mapping for Filter (sets) */
 	UPROPERTY()
-	TMap<UDataSourceFilter*, UDataSourceFilterSet*> ChildToParent;
+	TMap<TObjectPtr<UDataSourceFilter>, TObjectPtr<UDataSourceFilterSet>> ChildToParent;
 	
 	FSimpleMulticastDelegate SourceFiltersUpdatedDelegate;
 };
\ No newline at end of file
