make insights view toolbar extensible, and add a button for trace source filtering

diff --git a/Engine/Plugins/Developer/TraceSourceFiltering/Source/SourceFilteringEditor/Private/SourceFilterStyle.cpp b/Engine/Plugins/Developer/TraceSourceFiltering/Source/SourceFilteringEditor/Private/SourceFilterStyle.cpp
--- a/Engine/Plugins/Developer/TraceSourceFiltering/Source/SourceFilteringEditor/Private/SourceFilterStyle.cpp
+++ b/Engine/Plugins/Developer/TraceSourceFiltering/Source/SourceFilteringEditor/Private/SourceFilterStyle.cpp
@@ -149,7 +149,8 @@ void FSourceFilterStyle::Initialize()
 		.SetMenuBorderPadding(FMargin(0.0f));
 	StyleSet->Set("SourceFilter.ComboButton", ToolbarComboButton);
 
-	StyleSet->Set("SourceFilter.TabIcon", new IMAGE_BRUSH("Icons/AssetIcons/Character_16x", Icon16x16));
+	StyleSet->Set("SourceFilter.ToolBarIcon", new IMAGE_BRUSH_SVG("Starship/AssetIcons/Character_64", Icon32x32));
+	StyleSet->Set("SourceFilter.TabIcon", new IMAGE_BRUSH_SVG("Starship/AssetIcons/Character_16", Icon16x16));
 
 	StyleSet->Set("SourceFilter.NewPreset", new IMAGE_BRUSH("Icons/icon_file_saveas_16px", Icon16x16));
 	StyleSet->Set("SourceFilter.LoadPreset", new IMAGE_BRUSH("Icons/icon_file_open_16px", Icon16x16));
diff --git a/Engine/Plugins/Developer/TraceSourceFiltering/Source/SourceFilteringEditor/Private/SourceFilteringEditorModule.cpp b/Engine/Plugins/Developer/TraceSourceFiltering/Source/SourceFilteringEditor/Private/SourceFilteringEditorModule.cpp
--- a/Engine/Plugins/Developer/TraceSourceFiltering/Source/SourceFilteringEditor/Private/SourceFilteringEditorModule.cpp
+++ b/Engine/Plugins/Developer/TraceSourceFiltering/Source/SourceFilteringEditor/Private/SourceFilteringEditorModule.cpp
@@ -6,6 +6,7 @@
 #include "TraceServices/ITraceServicesModule.h"
 #include "HAL/IConsoleManager.h"
 #include "Insights/IUnrealInsightsModule.h"
+#include "Framework/Commands/Commands.h"
 #include "Framework/Docking/LayoutExtender.h"
 #include "Framework/Docking/TabManager.h"
 #include "Widgets/Docking/SDockTab.h"
@@ -32,9 +33,31 @@ IMPLEMENT_MODULE(FSourceFilteringEditorModule, SourceFilteringEditor);
 FString FSourceFilteringEditorModule::SourceFiltersIni;
 FName InsightsSourceFilteringTabName(TEXT("InsightsSourceFiltering"));
 
+#define LOCTEXT_NAMESPACE "FSourceFilteringEditorModule"
+
+class FSourceFilteringCommands : public TCommands<FSourceFilteringCommands>
+{
+public:
+	FSourceFilteringCommands()
+		: TCommands<FSourceFilteringCommands>("TraceSourceFiltering", LOCTEXT("TraceSourceFiltering", "Trace Source Filtering"), NAME_None, "SourceFilter")
+	{ }
+
+	virtual void RegisterCommands() override
+	{
+		UI_COMMAND( ToggleSourceFilteringVisibility,
+			"Source Filters",
+			"Opens the Trace Source Filtering tab, allows for filtering UWorld and AActor instances to not output Trace data",
+			EUserInterfaceActionType::ToggleButton, FInputChord());
+	}
+	
+	TSharedPtr<FUICommandInfo> ToggleSourceFilteringVisibility;
+};
+
+
 void FSourceFilteringEditorModule::StartupModule()
 {
 	FSourceFilterStyle::Initialize();
+	FSourceFilteringCommands::Register();
 
 	// Populate static ini path
 	FConfigCacheIni::LoadGlobalIniFile(SourceFiltersIni, TEXT("TraceSourceFilters"));
@@ -55,6 +78,26 @@ void FSourceFilteringEditorModule::ShutdownModule()
 	FSourceFilterStyle::Shutdown();
 }
 
+
+
+void FSourceFilteringEditorModule::ToggleSourceFilteringVisibility()
+{
+	TSharedPtr<SDockTab> SourceFilterTab = InsightsTabManager->FindExistingLiveTab(InsightsSourceFilteringTabName);
+	if (SourceFilterTab.IsValid())
+	{
+		SourceFilterTab->RequestCloseTab();
+	}
+	else
+	{
+		InsightsTabManager->TryInvokeTab(InsightsSourceFilteringTabName);
+	}
+}
+
+bool FSourceFilteringEditorModule::IsSourceFilteringVisibile()
+{
+	return bIsSourceFilterTabOpen;
+}
+
 void FSourceFilteringEditorModule::RegisterLayoutExtensions(FInsightsMajorTabExtender& InOutExtender)
 {
 #if WITH_EDITOR
@@ -63,25 +106,45 @@ void FSourceFilteringEditorModule::RegisterLayoutExtensions(FInsightsMajorTabExt
 	FTabId ExtendedTabId(FTimingProfilerTabs::TimersID);
 #endif
 
+	InsightsTabManager = InOutExtender.GetTabManager();
+	TSharedPtr<FUICommandList> CommandList = MakeShared<FUICommandList>();
+	const FSourceFilteringCommands& Commands = FSourceFilteringCommands::Get();
+
+	CommandList->MapAction(Commands.ToggleSourceFilteringVisibility, 
+							FExecuteAction::CreateRaw(this, &FSourceFilteringEditorModule::ToggleSourceFilteringVisibility), 
+							FCanExecuteAction(),
+							FIsActionChecked::CreateRaw(this, &FSourceFilteringEditorModule::IsSourceFilteringVisibile));
+	
+	InOutExtender.GetMenuExtender()->AddToolBarExtension( "MainToolbar", EExtensionHook::First, CommandList,
+		FToolBarExtensionDelegate::CreateLambda([](FToolBarBuilder& ToolBarBuilder)
+		{
+			ToolBarBuilder.AddToolBarButton(FSourceFilteringCommands::Get().ToggleSourceFilteringVisibility,
+				NAME_None, TAttribute<FText>(), TAttribute<FText>(),
+				FSlateIcon(FSourceFilterStyle::GetStyleSetName(), "SourceFilter.ToolBarIcon"));
+		}));
+	
 	InOutExtender.GetLayoutExtender().ExtendLayout(ExtendedTabId, ELayoutExtensionPosition::Before, FTabManager::FTab(InsightsSourceFilteringTabName, ETabState::ClosedTab));
 
 	TSharedRef<FWorkspaceItem> Category = InOutExtender.GetTabManager()->AddLocalWorkspaceMenuCategory(NSLOCTEXT("FInsightsSourceFilteringModule", "SourceFilteringGroupName", "Filtering"));
 
 	FMinorTabConfig& MinorTabConfig = InOutExtender.AddMinorTabConfig();
 	MinorTabConfig.TabId = InsightsSourceFilteringTabName;
-	MinorTabConfig.TabLabel = NSLOCTEXT("SourceFilteringEditorModule", "SourceFilteringTab", "Trace Source Filtering");
-	MinorTabConfig.TabTooltip = NSLOCTEXT("SourceFilteringEditorModule", "SourceFilteringTabTooltip", "Opens the Trace Source Filtering tab, allows for filtering UWorld and AActor instances to not output Trace data");
+	MinorTabConfig.TabLabel = LOCTEXT("SourceFilteringTab", "Trace Source Filtering");
+	MinorTabConfig.TabTooltip = LOCTEXT("SourceFilteringTabTooltip", "Opens the Trace Source Filtering tab, allows for filtering UWorld and AActor instances to not output Trace data");
 	MinorTabConfig.TabIcon = FSlateIcon(FSourceFilterStyle::GetStyleSetName(), "SourceFilter.TabIcon");
 	MinorTabConfig.WorkspaceGroup = Category;
-	MinorTabConfig.OnSpawnTab = FOnSpawnTab::CreateLambda([](const FSpawnTabArgs& Args)
+	MinorTabConfig.OnSpawnTab = FOnSpawnTab::CreateLambda([this](const FSpawnTabArgs& Args)
 	{
-		const TSharedRef<SDockTab> DockTab = SNew(SDockTab)
-			.TabRole(ETabRole::PanelTab);
+		bIsSourceFilterTabOpen = true;
+		
+		const TSharedRef<SDockTab> SourceFilterTab = SNew(SDockTab)
+			.TabRole(ETabRole::PanelTab)
+			.OnTabClosed_Lambda([this](TSharedRef<SDockTab>){ bIsSourceFilterTabOpen = false; });
 
 		TSharedRef<STraceSourceFilteringWidget> Window = SNew(STraceSourceFilteringWidget);
-		DockTab->SetContent(Window);
+		SourceFilterTab->SetContent(Window);
 
-		return DockTab;
+		return SourceFilterTab;
 	});
 }
 
@@ -144,4 +207,6 @@ void FSourceFilteringEditorModule::OnAssetsPendingDelete(TArray<UObject*> const&
 
 
 }
-#endif // WITH_EDITOR
\ No newline at end of file
+#endif // WITH_EDITOR
+
+#undef LOCTEXT_NAMESPACE
\ No newline at end of file
diff --git a/Engine/Plugins/Developer/TraceSourceFiltering/Source/SourceFilteringEditor/Private/SourceFilteringEditorModule.h b/Engine/Plugins/Developer/TraceSourceFiltering/Source/SourceFilteringEditor/Private/SourceFilteringEditorModule.h
--- a/Engine/Plugins/Developer/TraceSourceFiltering/Source/SourceFilteringEditor/Private/SourceFilteringEditorModule.h
+++ b/Engine/Plugins/Developer/TraceSourceFiltering/Source/SourceFilteringEditor/Private/SourceFilteringEditorModule.h
@@ -4,6 +4,7 @@
 
 #include "CoreMinimal.h"
 #include "Modules/ModuleInterface.h"
+#include "Framework/Docking/TabManager.h"
 
 #include "Templates/SharedPointer.h"
 
@@ -19,6 +20,12 @@ public:
 
 protected:
 	void RegisterLayoutExtensions(FInsightsMajorTabExtender& InOutExtender);
+	
+	void ToggleSourceFilteringVisibility();
+	bool IsSourceFilteringVisibile();
+
+	TSharedPtr<FTabManager> InsightsTabManager;
+	bool bIsSourceFilterTabOpen = false;
 
 #if WITH_EDITOR
 	/** Keeps track of any source filters due to be removed because their class (blueprint) is due to be deleted */
diff --git a/Engine/Source/Developer/TraceInsights/Private/Insights/Widgets/STimingProfilerToolbar.cpp b/Engine/Source/Developer/TraceInsights/Private/Insights/Widgets/STimingProfilerToolbar.cpp
--- a/Engine/Source/Developer/TraceInsights/Private/Insights/Widgets/STimingProfilerToolbar.cpp
+++ b/Engine/Source/Developer/TraceInsights/Private/Insights/Widgets/STimingProfilerToolbar.cpp
@@ -33,7 +33,7 @@ void STimingProfilerToolbar::Construct(const FArguments& InArgs)
 {
 	struct Local
 	{
-		static void FillViewToolbar(FToolBarBuilder& ToolbarBuilder)
+		static void FillViewToolbar(FToolBarBuilder& ToolbarBuilder, const FArguments &InArgs)
 		{
 			ToolbarBuilder.BeginSection("View");
 			{
@@ -58,6 +58,12 @@ void STimingProfilerToolbar::Construct(const FArguments& InArgs)
 				ToolbarBuilder.AddToolBarButton(FTimingProfilerCommands::Get().ToggleLogViewVisibility,
 					NAME_None, TAttribute<FText>(), TAttribute<FText>(),
 					FSlateIcon(FInsightsStyle::GetStyleSetName(), "Icons.LogView.ToolBar"));
+
+				if (InArgs._ToolbarExtender.IsValid())
+				{
+					InArgs._ToolbarExtender->Apply("MainToolbar", EExtensionHook::First, ToolbarBuilder);
+				}
+
 			}
 			ToolbarBuilder.EndSection();
 		}
@@ -78,7 +84,7 @@ void STimingProfilerToolbar::Construct(const FArguments& InArgs)
 
 	FSlimHorizontalToolBarBuilder ToolbarBuilder(CommandList.ToSharedRef(), FMultiBoxCustomization::None);
 	ToolbarBuilder.SetStyle(&FInsightsStyle::Get(), "PrimaryToolbar");
-	Local::FillViewToolbar(ToolbarBuilder);
+	Local::FillViewToolbar(ToolbarBuilder, InArgs);
 
 	FSlimHorizontalToolBarBuilder RightSideToolbarBuilder(CommandList.ToSharedRef(), FMultiBoxCustomization::None);
 	RightSideToolbarBuilder.SetStyle(&FInsightsStyle::Get(), "PrimaryToolbar");
diff --git a/Engine/Source/Developer/TraceInsights/Private/Insights/Widgets/STimingProfilerToolbar.h b/Engine/Source/Developer/TraceInsights/Private/Insights/Widgets/STimingProfilerToolbar.h
--- a/Engine/Source/Developer/TraceInsights/Private/Insights/Widgets/STimingProfilerToolbar.h
+++ b/Engine/Source/Developer/TraceInsights/Private/Insights/Widgets/STimingProfilerToolbar.h
@@ -4,6 +4,7 @@
 
 #include "Widgets/DeclarativeSyntaxSupport.h"
 #include "Widgets/SCompoundWidget.h"
+#include "Framework/MultiBox/MultiBoxExtender.h"
 
 /** Ribbon based toolbar used as a main menu in the Profiler window. */
 class STimingProfilerToolbar : public SCompoundWidget
@@ -16,6 +17,7 @@ public:
 	virtual ~STimingProfilerToolbar();
 
 	SLATE_BEGIN_ARGS(STimingProfilerToolbar) {}
+	SLATE_ARGUMENT(TSharedPtr<FExtender>, ToolbarExtender);
 	SLATE_END_ARGS()
 
 	/**
diff --git a/Engine/Source/Developer/TraceInsights/Private/Insights/Widgets/STimingProfilerWindow.cpp b/Engine/Source/Developer/TraceInsights/Private/Insights/Widgets/STimingProfilerWindow.cpp
--- a/Engine/Source/Developer/TraceInsights/Private/Insights/Widgets/STimingProfilerWindow.cpp
+++ b/Engine/Source/Developer/TraceInsights/Private/Insights/Widgets/STimingProfilerWindow.cpp
@@ -370,10 +370,6 @@ void STimingProfilerWindow::Construct(const FArguments& InArgs, const TSharedRef
 	// Create & initialize tab manager.
 	TabManager = FGlobalTabmanager::Get()->NewTabManager(ConstructUnderMajorTab);
 
-#if WITH_EDITOR
-	TabManager->SetAllowWindowMenuBar(true);
-#endif
-
 	const auto& PersistLayout = [](const TSharedRef<FTabManager::FLayout>& LayoutToSave)
 	{
 		FLayoutSaveRestore::SaveToConfig(FTraceInsightsModule::GetUnrealInsightsLayoutIni(), LayoutToSave);
@@ -570,20 +566,12 @@ void STimingProfilerWindow::Construct(const FArguments& InArgs, const TSharedRef
 		.VAlign(VAlign_Fill)
 		[
 			SNew(SVerticalBox)
-
-#if WITH_EDITOR
-			+ SVerticalBox::Slot()
-			.AutoHeight()
-			[
-				MenuWidget
-			]
-#endif
-
+			
 			+ SVerticalBox::Slot()
 			.AutoHeight()
 			.Padding(FMargin(0.0f, 0.0f, 0.0f, 5.0f))
 			[
-				SNew(STimingProfilerToolbar)
+				SNew(STimingProfilerToolbar).ToolbarExtender(Extension->GetMenuExtender())
 			]
 
 			+ SVerticalBox::Slot()
