[KillPendingKill] Make FFilteredActorCollector not take ownership of actors.

diff --git a/Engine/Plugins/Developer/TraceSourceFiltering/Source/SourceFilteringTrace/Private/ActorFiltering.h b/Engine/Plugins/Developer/TraceSourceFiltering/Source/SourceFilteringTrace/Private/ActorFiltering.h
--- a/Engine/Plugins/Developer/TraceSourceFiltering/Source/SourceFilteringTrace/Private/ActorFiltering.h
+++ b/Engine/Plugins/Developer/TraceSourceFiltering/Source/SourceFilteringTrace/Private/ActorFiltering.h
@@ -29,9 +29,10 @@ public:
 
 	FORCEINLINE const AActor* GetActorChecked() const
 	{
-		check(CurrentActor);
-		checkf(!CurrentActor->IsUnreachable(), TEXT("%s"), *CurrentActor->GetFullName());
-		return CurrentActor;
+		const AActor* Actor = CurrentActor.Get();
+		check(Actor);
+		checkf(!Actor->IsUnreachable(), TEXT("%s"), *CurrentActor->GetFullName());
+		return Actor;
 	}
 
 	FORCEINLINE explicit operator bool() const
@@ -45,14 +46,14 @@ public:
 		while (++CurrentIndex < EndIndex)
 		{
 			CurrentActor = Actors[CurrentIndex];
-			if (IsValid(CurrentActor))
+			if (const AActor* Actor =  CurrentActor.Get(); IsValid(Actor))
 			{
 				break;
 			}
 		}
 	}
 private:
-	FFilteredActorIterator(const TArray<const AActor*>& InActors, uint32 InStartIndex, uint32 InLength)
+	FFilteredActorIterator(const TArray<TWeakObjectPtr<const AActor>>& InActors, uint32 InStartIndex, uint32 InLength)
 		: StartIndex(InStartIndex), CurrentIndex(InStartIndex), CurrentActor(nullptr), Actors(InActors)
 	{
 		if (InLength == 0)
@@ -70,6 +71,10 @@ private:
 			check(InActors.IsValidIndex(StartIndex));
 			check(InActors.IsValidIndex(EndIndex - 1));
 			CurrentActor = Actors[StartIndex];
+			if (const AActor* Actor = CurrentActor.Get(); !IsValid(Actor))
+			{
+				++(*this); // Advance til we hit the first valid actor. 
+			}
 		}
 	}
 		
@@ -78,10 +83,10 @@ private:
 	uint32 EndIndex;
 	/** Current index and corresponding AActor instance */
 	uint32 CurrentIndex;
-	const AActor* CurrentActor;
+	TWeakObjectPtr<const AActor> CurrentActor;
 
 	/** All Actors retrieved by FActorCollector, over which we iterate within the defined range */
-	const TArray<const AActor*>& Actors;
+	const TArray<TWeakObjectPtr<const AActor>>& Actors;
 };
 
 /** Used for retrieving AActor instances of particular classes from the owned UWorld, derived from FActorIteratorState */
@@ -202,7 +207,6 @@ public:
 
 	void AddReferencedObjects(FReferenceCollector& Collector) override
 	{
-		Collector.AddReferencedObjects(ActorArray);
 		Collector.AddReferencedObjects(AllFilteredClasses);
 	}
 	virtual FString GetReferencerName() const override
@@ -213,6 +217,6 @@ public:
 protected:
 	UWorld* CurrentWorld;
 
-	TArray<const AActor*> ActorArray;
+	TArray<TWeakObjectPtr<const AActor>> ActorArray;
 	TArray<const UClass*> AllFilteredClasses;
 };
\ No newline at end of file
