releasing static FGraphEvent on engine exit to avoid undefined order of destruction (can be destroyed after its allocator)

diff --git a/Engine/Plugins/Developer/TraceSourceFiltering/Source/SourceFilteringTrace/Private/SourceFilterManager.cpp b/Engine/Plugins/Developer/TraceSourceFiltering/Source/SourceFilteringTrace/Private/SourceFilterManager.cpp
--- a/Engine/Plugins/Developer/TraceSourceFiltering/Source/SourceFilteringTrace/Private/SourceFilterManager.cpp
+++ b/Engine/Plugins/Developer/TraceSourceFiltering/Source/SourceFilteringTrace/Private/SourceFilterManager.cpp
@@ -655,7 +655,10 @@ void FSourceFilterManager::DrawFilterResults(const AActor* Actor)
 
 void FSourceFilterManager::SetupAsyncTasks(ENamedThreads::Type CurrentThread)
 {
-	static FGraphEventRef LastApplyAsyncTaskEvent = nullptr;
+	static FGraphEventRef LastApplyAsyncTaskEvent;
+
+	// if the static var left alive it can be destroyed after its allocator
+	UE_CALL_ONCE([] { FCoreDelegates::OnEnginePreExit.AddLambda([] { LastApplyAsyncTaskEvent = nullptr; }); });
 
 	if (CAN_TRACE_OBJECT(World))
 	{
