add null check to FSourceFilterTrace::OutputClass

diff --git a/Engine/Plugins/Developer/TraceSourceFiltering/Source/SourceFilteringTrace/Private/SourceFilterTrace.cpp b/Engine/Plugins/Developer/TraceSourceFiltering/Source/SourceFilteringTrace/Private/SourceFilterTrace.cpp
--- a/Engine/Plugins/Developer/TraceSourceFiltering/Source/SourceFilteringTrace/Private/SourceFilterTrace.cpp
+++ b/Engine/Plugins/Developer/TraceSourceFiltering/Source/SourceFilteringTrace/Private/SourceFilterTrace.cpp
@@ -72,6 +72,11 @@ TMap<uint64, FObjectKey> FSourceFilterTrace::IDsToWorldInstance;
 
 void FSourceFilterTrace::OutputClass(const TSubclassOf<UDataSourceFilter> InClass)
 {
+	if (InClass == nullptr)
+	{
+		return;
+	}
+
 	// Skip SKEL and REINST classes.
 	if (InClass->GetName().StartsWith(TEXT("SKEL_")) || InClass->GetName().StartsWith(TEXT("REINST_")))
 	{
