fix some untagged allocations in GameplayInsights and SourceFiltering

diff --git a/Engine/Plugins/Animation/GameplayInsights/Source/GameplayInsights/Private/GameplayInsightsModule.cpp b/Engine/Plugins/Animation/GameplayInsights/Source/GameplayInsights/Private/GameplayInsightsModule.cpp
--- a/Engine/Plugins/Animation/GameplayInsights/Source/GameplayInsights/Private/GameplayInsightsModule.cpp
+++ b/Engine/Plugins/Animation/GameplayInsights/Source/GameplayInsights/Private/GameplayInsightsModule.cpp
@@ -422,6 +422,8 @@ bool FGameplayInsightsModule::IsObjectPropertyTraceEnabled(UObject* Object)
 
 void FGameplayInsightsModule::StartTrace()
 {
+	LLM_SCOPE_BYNAME(TEXT("Insights/GameplayInsights"));
+	
 	if (!bTraceStarted)
 	{
 		if (!FTraceAuxiliary::IsConnected())
diff --git a/Engine/Plugins/Developer/TraceSourceFiltering/Source/SourceFilteringEditor/Private/SourceFilteringEditorModule.cpp b/Engine/Plugins/Developer/TraceSourceFiltering/Source/SourceFilteringEditor/Private/SourceFilteringEditorModule.cpp
--- a/Engine/Plugins/Developer/TraceSourceFiltering/Source/SourceFilteringEditor/Private/SourceFilteringEditorModule.cpp
+++ b/Engine/Plugins/Developer/TraceSourceFiltering/Source/SourceFilteringEditor/Private/SourceFilteringEditorModule.cpp
@@ -57,6 +57,8 @@ public:
 
 void FSourceFilteringEditorModule::StartupModule()
 {
+	LLM_SCOPE_BYNAME(TEXT("Insights/SourceFilteringEditor"));
+	
 	FSourceFilterStyle::Initialize();
 	FSourceFilteringCommands::Register();
 
@@ -101,6 +103,7 @@ bool FSourceFilteringEditorModule::IsSourceFilteringVisibile()
 
 void FSourceFilteringEditorModule::RegisterLayoutExtensions(FInsightsMajorTabExtender& InOutExtender)
 {
+	LLM_SCOPE_BYNAME(TEXT("Insights/SourceFilteringEditor"));
 #if WITH_EDITOR
 	FTabId ExtendedTabId(GameplayInsightsTabs::DocumentTab);
 #else
