Stop force enabling TraceFiltering trace channel

diff --git a/Engine/Plugins/Developer/TraceSourceFiltering/Source/SourceFilteringTrace/Private/SourceFilteringTraceModule.cpp b/Engine/Plugins/Developer/TraceSourceFiltering/Source/SourceFilteringTrace/Private/SourceFilteringTraceModule.cpp
--- a/Engine/Plugins/Developer/TraceSourceFiltering/Source/SourceFilteringTrace/Private/SourceFilteringTraceModule.cpp
+++ b/Engine/Plugins/Developer/TraceSourceFiltering/Source/SourceFilteringTrace/Private/SourceFilteringTraceModule.cpp
@@ -76,9 +76,6 @@ void FSourceFilteringTraceModule::StartupModule()
 	FTraceWorldFiltering::Initialize();
 	FTraceSourceFiltering::Initialize();
 
-	// Forcefully enable the source trace channel
-	UE::Trace::ToggleChannel(TEXT("TraceSourceFiltersChannel"), true);
-
 #if WITH_EDITOR
 	// Add callback to trace out Filter Classes once the Asset Registry has finished loading 
 	FAssetRegistryModule& AssetRegistryModule = FModuleManager::LoadModuleChecked<FAssetRegistryModule>(TEXT("AssetRegistry"));
