Introduces a new world type "EditorStorage".

The "EditorStorage" world is meant to be used to store data for the entire lifetime of the editor process. As such it's not destroyed when new levels are loaded and is meant to be constructed with minimal setup to avoid overhead. This new world type is currently used in the Typed Elements Data Storage which serves as the central database-like storage for the editor and is powered by MASS. The latter currently requires a world for interacting with the outside through subsystems and to be periodically ticked to make sure data remains consistent. The changes that make use of this will be submitted separately to keep this CL focused.

diff --git a/Engine/Plugins/Developer/TraceSourceFiltering/Source/SourceFilteringEditor/Private/WorldFilters.cpp b/Engine/Plugins/Developer/TraceSourceFiltering/Source/SourceFilteringEditor/Private/WorldFilters.cpp
--- a/Engine/Plugins/Developer/TraceSourceFiltering/Source/SourceFilteringEditor/Private/WorldFilters.cpp
+++ b/Engine/Plugins/Developer/TraceSourceFiltering/Source/SourceFilteringEditor/Private/WorldFilters.cpp
@@ -20,17 +20,18 @@
 FWorldTypeTraceFilter::FWorldTypeTraceFilter(const TFunction<void(uint8, bool)>& InOnSetWorldTypeFilterState, const TFunction<bool(uint8)>& InOnGetWorldTypeFilterState) : OnSetWorldTypeFilterState(InOnSetWorldTypeFilterState), OnGetWorldTypeFilterState(InOnGetWorldTypeFilterState)
 {
 	WorldTypeFilterName = LOCTEXT("WorldTypeFilterName", "By World Type");
-	WorldTypeFilterValues.Add(TPair<FText, TArray<uint8>>(LOCTEXT("EditorWorldFilterName", "Editor"), TArray<uint8>({ 2 /*EWorldType::Editor*/, 4 /*EWorldType::EditorPreview*/ })));
+	WorldTypeFilterValues.Add(TPair<FText, TArray<uint8>>(LOCTEXT("EditorWorldFilterName", "Editor"), TArray<uint8>({ 2 /*EWorldType::Editor*/, 4 /*EWorldType::EditorPreview*/, 5 /*EWorldType::EditorStorage*/ })));
 	WorldTypeFilterValues.Add(TPair<FText, TArray<uint8>>(LOCTEXT("PIEWorldFilterName", "PIE"), TArray<uint8>({ 3 /*EWorldType::PIE*/ })));
-	WorldTypeFilterValues.Add(TPair<FText, TArray<uint8>>(LOCTEXT("GameWorldFilterName", "Game"), TArray<uint8>({ 1 /*EWorldType::Game*/, 5 /*EWorldType::GamePreview*/, 6 /*EWorldType::GameRPC*/ })));
+	WorldTypeFilterValues.Add(TPair<FText, TArray<uint8>>(LOCTEXT("GameWorldFilterName", "Game"), TArray<uint8>({ 1 /*EWorldType::Game*/, 6 /*EWorldType::GamePreview*/, 7 /*EWorldType::GameRPC*/ })));
 
 #if WITH_ENGINE
 	static_assert((uint8)EWorldType::Editor == 2, "EWorldType and WorldTypeFilterValues must be kept in sync");
 	static_assert((uint8)EWorldType::EditorPreview == 4, "EWorldType and WorldTypeFilterValues must be kept in sync");
+	static_assert((uint8)EWorldType::EditorStorage == 5, "EWorldType and WorldTypeFilterValues must be kept in sync");
 	static_assert((uint8)EWorldType::PIE == 3, "EWorldType and WorldTypeFilterValues must be kept in sync");
 	static_assert((uint8)EWorldType::Game == 1, "EWorldType and WorldTypeFilterValues must be kept in sync");
-	static_assert((uint8)EWorldType::GamePreview == 5, "EWorldType and WorldTypeFilterValues must be kept in sync");
-	static_assert((uint8)EWorldType::GameRPC == 6, "EWorldType and WorldTypeFilterValues must be kept in sync");
+	static_assert((uint8)EWorldType::GamePreview == 6, "EWorldType and WorldTypeFilterValues must be kept in sync");
+	static_assert((uint8)EWorldType::GameRPC == 7, "EWorldType and WorldTypeFilterValues must be kept in sync");
 #endif
 
 	LoadSettings();
diff --git a/Engine/Source/Editor/UnrealEd/Private/EditorServer.cpp b/Engine/Source/Editor/UnrealEd/Private/EditorServer.cpp
--- a/Engine/Source/Editor/UnrealEd/Private/EditorServer.cpp
+++ b/Engine/Source/Editor/UnrealEd/Private/EditorServer.cpp
@@ -2034,7 +2034,11 @@ void UEditorEngine::CheckForWorldGCLeaks( UWorld* NewWorld, UPackage* WorldPacka
 	{
 		UWorld* RemainingWorld = *It;
 		const bool bIsNewWorld = (NewWorld && RemainingWorld == NewWorld);
-		const bool bIsPersistantWorldType = (RemainingWorld->WorldType == EWorldType::Inactive) || (RemainingWorld->WorldType == EWorldType::EditorPreview) || (RemainingWorld->WorldType == EWorldType::GamePreview);
+		const bool bIsPersistantWorldType = (
+			RemainingWorld->WorldType == EWorldType::Inactive) || 
+			(RemainingWorld->WorldType == EWorldType::EditorPreview) || 
+			(RemainingWorld->WorldType == EWorldType::GamePreview) ||
+			(RemainingWorld->WorldType == EWorldType::EditorStorage);
 		if(!bIsNewWorld && !bIsPersistantWorldType && !WorldHasValidContext(RemainingWorld))
 		{
 			FindAndPrintStaleReferencesToObject(RemainingWorld, EPrintStaleReferencesOptions::Error);
@@ -2088,7 +2092,9 @@ void UEditorEngine::EditorDestroyWorld( FWorldContext & Context, const FText& Cl
 		WorldPackage = NULL;
 	}
 
-	if (ContextWorld->WorldType != EWorldType::EditorPreview && ContextWorld->WorldType != EWorldType::Inactive)
+	if (ContextWorld->WorldType != EWorldType::EditorPreview && 
+		ContextWorld->WorldType != EWorldType::EditorStorage &&
+		ContextWorld->WorldType != EWorldType::Inactive)
 	{
 		// Go away, come again never!
 		ContextWorld->ClearFlags(RF_Standalone | RF_Transactional);
diff --git a/Engine/Source/Runtime/Engine/Classes/Engine/EngineTypes.h b/Engine/Source/Runtime/Engine/Classes/Engine/EngineTypes.h
--- a/Engine/Source/Runtime/Engine/Classes/Engine/EngineTypes.h
+++ b/Engine/Source/Runtime/Engine/Classes/Engine/EngineTypes.h
@@ -1098,6 +1098,13 @@ namespace EWorldType
 		/** A preview world for an editor tool */
 		EditorPreview,
 
+		/** 
+		 * A minimal world for storing editor information.
+		 * Alert: This world type is part of experimental development and it's strongly recommended to not use this world type
+		 *		as long a this alert is here.
+		 */
+		EditorStorage,
+
 		/** A preview world for a game */
 		GamePreview,
 
diff --git a/Engine/Source/Runtime/Engine/Classes/Engine/World.h b/Engine/Source/Runtime/Engine/Classes/Engine/World.h
--- a/Engine/Source/Runtime/Engine/Classes/Engine/World.h
+++ b/Engine/Source/Runtime/Engine/Classes/Engine/World.h
@@ -3718,6 +3718,13 @@ public:
 	/** Returns true if this world is any kind of editor world (including editor preview worlds) */
 	bool IsEditorWorld() const;
 
+	/** 
+	 * Returns true if this is a minimal world for storing information.
+	 * Alert: This is part of experimental development and it's strongly recommended to not use this function
+	 *		as long a this alert is here.
+	 */
+	bool IsStorageWorld() const;
+
 	/** Returns true if this world is a preview game world (editor or game) */
 	bool IsPreviewWorld() const;
 
diff --git a/Engine/Source/Runtime/Engine/Private/World.cpp b/Engine/Source/Runtime/Engine/Private/World.cpp
--- a/Engine/Source/Runtime/Engine/Private/World.cpp
+++ b/Engine/Source/Runtime/Engine/Private/World.cpp
@@ -7747,7 +7747,7 @@ bool UWorld::IsGameWorld() const
 
 bool UWorld::IsEditorWorld() const
 {
-	return WorldType == EWorldType::Editor || WorldType == EWorldType::EditorPreview || WorldType == EWorldType::PIE;
+	return WorldType == EWorldType::Editor || WorldType == EWorldType::EditorPreview || WorldType == EWorldType::PIE || WorldType == EWorldType::EditorStorage;
 }
 
 bool UWorld::IsPreviewWorld() const
@@ -7755,6 +7755,11 @@ bool UWorld::IsPreviewWorld() const
 	return WorldType == EWorldType::EditorPreview || WorldType == EWorldType::GamePreview;
 }
 
+bool UWorld::IsStorageWorld() const
+{
+	return WorldType == EWorldType::EditorStorage;
+}
+
 bool UWorld::UsesGameHiddenFlags() const
 {
 	return IsGameWorld();
@@ -9028,6 +9033,7 @@ FString ToString(EWorldType::Type Type)
 	case EWorldType::Editor: return TEXT("Editor");
 	case EWorldType::PIE: return TEXT("PIE");
 	case EWorldType::EditorPreview: return TEXT("EditorPreview");
+	case EWorldType::EditorStorage: return TEXT("EditorStorage");
 	case EWorldType::GamePreview: return TEXT("GamePreview");
 	case EWorldType::GameRPC: return TEXT("GameRPC");
 	case EWorldType::Inactive: return TEXT("Inactive");
