Removed the EditorStorage world type.

The EditorStorage world type was introduced as MASS had a dependency on  having a world, but none of the existing worlds could be made to work. This limitation has since been removed from MASS which means there's no current use for EditorStorage. To avoid adding additional maintenance and since EditorStorage was introduce within 5.2, it's now being removed before release. If a use case is found at a later time, this CL can be used to reintroduce EditorStorage.

diff --git a/Engine/Plugins/Developer/TraceSourceFiltering/Source/SourceFilteringEditor/Private/WorldFilters.cpp b/Engine/Plugins/Developer/TraceSourceFiltering/Source/SourceFilteringEditor/Private/WorldFilters.cpp
--- a/Engine/Plugins/Developer/TraceSourceFiltering/Source/SourceFilteringEditor/Private/WorldFilters.cpp
+++ b/Engine/Plugins/Developer/TraceSourceFiltering/Source/SourceFilteringEditor/Private/WorldFilters.cpp
@@ -19,18 +19,17 @@
 FWorldTypeTraceFilter::FWorldTypeTraceFilter(const TFunction<void(uint8, bool)>& InOnSetWorldTypeFilterState, const TFunction<bool(uint8)>& InOnGetWorldTypeFilterState) : OnSetWorldTypeFilterState(InOnSetWorldTypeFilterState), OnGetWorldTypeFilterState(InOnGetWorldTypeFilterState)
 {
 	WorldTypeFilterName = LOCTEXT("WorldTypeFilterName", "By World Type");
-	WorldTypeFilterValues.Add(TPair<FText, TArray<uint8>>(LOCTEXT("EditorWorldFilterName", "Editor"), TArray<uint8>({ 2 /*EWorldType::Editor*/, 4 /*EWorldType::EditorPreview*/, 5 /*EWorldType::EditorStorage*/ })));
+	WorldTypeFilterValues.Add(TPair<FText, TArray<uint8>>(LOCTEXT("EditorWorldFilterName", "Editor"), TArray<uint8>({ 2 /*EWorldType::Editor*/, 4 /*EWorldType::EditorPreview*/ })));
 	WorldTypeFilterValues.Add(TPair<FText, TArray<uint8>>(LOCTEXT("PIEWorldFilterName", "PIE"), TArray<uint8>({ 3 /*EWorldType::PIE*/ })));
-	WorldTypeFilterValues.Add(TPair<FText, TArray<uint8>>(LOCTEXT("GameWorldFilterName", "Game"), TArray<uint8>({ 1 /*EWorldType::Game*/, 6 /*EWorldType::GamePreview*/, 7 /*EWorldType::GameRPC*/ })));
+	WorldTypeFilterValues.Add(TPair<FText, TArray<uint8>>(LOCTEXT("GameWorldFilterName", "Game"), TArray<uint8>({ 1 /*EWorldType::Game*/, 5 /*EWorldType::GamePreview*/, 6 /*EWorldType::GameRPC*/ })));
 
 #if WITH_ENGINE
 	static_assert((uint8)EWorldType::Editor == 2, "EWorldType and WorldTypeFilterValues must be kept in sync");
 	static_assert((uint8)EWorldType::EditorPreview == 4, "EWorldType and WorldTypeFilterValues must be kept in sync");
-	static_assert((uint8)EWorldType::EditorStorage == 5, "EWorldType and WorldTypeFilterValues must be kept in sync");
 	static_assert((uint8)EWorldType::PIE == 3, "EWorldType and WorldTypeFilterValues must be kept in sync");
 	static_assert((uint8)EWorldType::Game == 1, "EWorldType and WorldTypeFilterValues must be kept in sync");
-	static_assert((uint8)EWorldType::GamePreview == 6, "EWorldType and WorldTypeFilterValues must be kept in sync");
-	static_assert((uint8)EWorldType::GameRPC == 7, "EWorldType and WorldTypeFilterValues must be kept in sync");
+	static_assert((uint8)EWorldType::GamePreview == 5, "EWorldType and WorldTypeFilterValues must be kept in sync");
+	static_assert((uint8)EWorldType::GameRPC == 6, "EWorldType and WorldTypeFilterValues must be kept in sync");
 #endif
 
 	LoadSettings();
diff --git a/Engine/Plugins/Runtime/MassEntity/Source/MassEntity/Private/MassEntitySubsystem.cpp b/Engine/Plugins/Runtime/MassEntity/Source/MassEntity/Private/MassEntitySubsystem.cpp
--- a/Engine/Plugins/Runtime/MassEntity/Source/MassEntity/Private/MassEntitySubsystem.cpp
+++ b/Engine/Plugins/Runtime/MassEntity/Source/MassEntity/Private/MassEntitySubsystem.cpp
@@ -40,11 +40,6 @@ void UMassEntitySubsystem::Deinitialize()
 	EntityManager.Reset();
 }
 
-bool UMassEntitySubsystem::DoesSupportWorldType(const EWorldType::Type WorldType) const
-{
-	return Super::DoesSupportWorldType(WorldType);
-}
-
 #if WITH_MASSENTITY_DEBUG
 //////////////////////////////////////////////////////////////////////
 // Debug commands
diff --git a/Engine/Plugins/Runtime/MassEntity/Source/MassEntity/Public/MassEntitySubsystem.h b/Engine/Plugins/Runtime/MassEntity/Source/MassEntity/Public/MassEntitySubsystem.h
--- a/Engine/Plugins/Runtime/MassEntity/Source/MassEntity/Public/MassEntitySubsystem.h
+++ b/Engine/Plugins/Runtime/MassEntity/Source/MassEntity/Public/MassEntitySubsystem.h
@@ -23,11 +23,6 @@ class MASSENTITY_API UMassEntitySubsystem : public UWorldSubsystem
 	virtual void Deinitialize() override;
 	//~End of USubsystem interface
 
-protected:
-	//~USubsystem interface
-	virtual bool DoesSupportWorldType(const EWorldType::Type WorldType) const override;
-	//~End of USubsystem interface
-
 public:
 	UMassEntitySubsystem();
 
diff --git a/Engine/Plugins/Runtime/MassGameplay/Source/MassActors/Private/MassActorSubsystem.cpp b/Engine/Plugins/Runtime/MassGameplay/Source/MassActors/Private/MassActorSubsystem.cpp
--- a/Engine/Plugins/Runtime/MassGameplay/Source/MassActors/Private/MassActorSubsystem.cpp
+++ b/Engine/Plugins/Runtime/MassGameplay/Source/MassActors/Private/MassActorSubsystem.cpp
@@ -95,11 +95,6 @@ void UMassActorSubsystem::Deinitialize()
 	ActorManager.Reset();
 }
 
-bool UMassActorSubsystem::DoesSupportWorldType(const EWorldType::Type WorldType) const
-{
-	return Super::DoesSupportWorldType(WorldType);
-}
-
 //----------------------------------------------------------------------//
 //  FMassActorManager
 //----------------------------------------------------------------------//
diff --git a/Engine/Plugins/Runtime/MassGameplay/Source/MassActors/Public/MassActorSubsystem.h b/Engine/Plugins/Runtime/MassGameplay/Source/MassActors/Public/MassActorSubsystem.h
--- a/Engine/Plugins/Runtime/MassGameplay/Source/MassActors/Public/MassActorSubsystem.h
+++ b/Engine/Plugins/Runtime/MassGameplay/Source/MassActors/Public/MassActorSubsystem.h
@@ -148,7 +148,6 @@ protected:
 	// USubsystem BEGIN
 	virtual void Initialize(FSubsystemCollectionBase& Collection) override;
 	virtual void Deinitialize() override;
-	virtual bool DoesSupportWorldType(const EWorldType::Type WorldType) const override;
 	// USubsystem END
 	
 	TSharedPtr<FMassActorManager> ActorManager;
diff --git a/Engine/Plugins/Runtime/MassGameplay/Source/MassGameplayDebug/Private/MassDebuggerSubsystem.cpp b/Engine/Plugins/Runtime/MassGameplay/Source/MassGameplayDebug/Private/MassDebuggerSubsystem.cpp
--- a/Engine/Plugins/Runtime/MassGameplay/Source/MassGameplayDebug/Private/MassDebuggerSubsystem.cpp
+++ b/Engine/Plugins/Runtime/MassGameplay/Source/MassGameplayDebug/Private/MassDebuggerSubsystem.cpp
@@ -32,11 +32,6 @@ void UMassDebuggerSubsystem::Deinitialize()
 #endif // WITH_MASSENTITY_DEBUG
 }
 
-bool UMassDebuggerSubsystem::DoesSupportWorldType(const EWorldType::Type WorldType) const
-{
-	return Super::DoesSupportWorldType(WorldType);
-}
-
 #if WITH_MASSENTITY_DEBUG
 void UMassDebuggerSubsystem::OnEntitySelected(const FMassEntityManager& EntityManager, const FMassEntityHandle EntityHandle)
 {
diff --git a/Engine/Plugins/Runtime/MassGameplay/Source/MassGameplayDebug/Public/MassDebuggerSubsystem.h b/Engine/Plugins/Runtime/MassGameplay/Source/MassGameplayDebug/Public/MassDebuggerSubsystem.h
--- a/Engine/Plugins/Runtime/MassGameplay/Source/MassGameplayDebug/Public/MassDebuggerSubsystem.h
+++ b/Engine/Plugins/Runtime/MassGameplay/Source/MassGameplayDebug/Public/MassDebuggerSubsystem.h
@@ -47,7 +47,6 @@ protected:
 	// USubsystem BEGIN
 	virtual void Initialize(FSubsystemCollectionBase& Collection) override;
 	virtual void Deinitialize() override;
-	virtual bool DoesSupportWorldType(const EWorldType::Type WorldType) const override;
 	// USubsystem END
 	
 	void OnProcessingPhaseStarted(const float DeltaSeconds);
diff --git a/Engine/Plugins/Runtime/MassGameplay/Source/MassSimulation/Private/MassSimulationSubsystem.cpp b/Engine/Plugins/Runtime/MassGameplay/Source/MassSimulation/Private/MassSimulationSubsystem.cpp
--- a/Engine/Plugins/Runtime/MassGameplay/Source/MassSimulation/Private/MassSimulationSubsystem.cpp
+++ b/Engine/Plugins/Runtime/MassGameplay/Source/MassSimulation/Private/MassSimulationSubsystem.cpp
@@ -56,11 +56,6 @@ void UMassSimulationSubsystem::BeginDestroy()
 	Super::BeginDestroy();
 }
 
-bool UMassSimulationSubsystem::DoesSupportWorldType(const EWorldType::Type WorldType) const
-{
-	return Super::DoesSupportWorldType(WorldType);
-}
-
 FMassProcessingPhase::FOnPhaseEvent& UMassSimulationSubsystem::GetOnProcessingPhaseStarted(const EMassProcessingPhase Phase)
 {
 	return PhaseManager.GetOnPhaseStart(Phase);
diff --git a/Engine/Plugins/Runtime/MassGameplay/Source/MassSimulation/Public/MassSimulationSubsystem.h b/Engine/Plugins/Runtime/MassGameplay/Source/MassSimulation/Public/MassSimulationSubsystem.h
--- a/Engine/Plugins/Runtime/MassGameplay/Source/MassSimulation/Public/MassSimulationSubsystem.h
+++ b/Engine/Plugins/Runtime/MassGameplay/Source/MassSimulation/Public/MassSimulationSubsystem.h
@@ -42,8 +42,7 @@ protected:
 	virtual void OnWorldBeginPlay(UWorld& InWorld) override;
 	virtual void Deinitialize() override;
 	virtual void BeginDestroy() override;
-	virtual bool DoesSupportWorldType(const EWorldType::Type WorldType) const override;
-
+	
 	void RebuildTickPipeline();
 
 	void StartSimulation(UWorld& InWorld);
diff --git a/Engine/Plugins/Runtime/ZoneGraph/Source/ZoneGraph/Private/ZoneGraphSubsystem.cpp b/Engine/Plugins/Runtime/ZoneGraph/Source/ZoneGraph/Private/ZoneGraphSubsystem.cpp
--- a/Engine/Plugins/Runtime/ZoneGraph/Source/ZoneGraph/Private/ZoneGraphSubsystem.cpp
+++ b/Engine/Plugins/Runtime/ZoneGraph/Source/ZoneGraph/Private/ZoneGraphSubsystem.cpp
@@ -493,8 +493,3 @@ TConstArrayView<FZoneGraphTagInfo> UZoneGraphSubsystem::GetTagInfos() const
 	}
 	return TConstArrayView<FZoneGraphTagInfo>();
 }
-
-bool UZoneGraphSubsystem::DoesSupportWorldType(const EWorldType::Type WorldType) const
-{
-	return WorldType == EWorldType::EditorStorage || Super::DoesSupportWorldType(WorldType);
-}
\ No newline at end of file
diff --git a/Engine/Plugins/Runtime/ZoneGraph/Source/ZoneGraph/Public/ZoneGraphSubsystem.h b/Engine/Plugins/Runtime/ZoneGraph/Source/ZoneGraph/Public/ZoneGraphSubsystem.h
--- a/Engine/Plugins/Runtime/ZoneGraph/Source/ZoneGraph/Public/ZoneGraphSubsystem.h
+++ b/Engine/Plugins/Runtime/ZoneGraph/Source/ZoneGraph/Public/ZoneGraphSubsystem.h
@@ -134,8 +134,7 @@ protected:
 	virtual void Deinitialize() override;
 	virtual void Tick(float DeltaTime) override;
 	virtual TStatId GetStatId() const override;
-	virtual bool DoesSupportWorldType(const EWorldType::Type WorldType) const override;
-
+	
 	void RemoveRegisteredDataItem(const int32 Index);
 	void UnregisterStaleZoneGraphDataInstances();
 	void RegisterZoneGraphDataInstances();
diff --git a/Engine/Plugins/Runtime/ZoneGraphAnnotations/Source/ZoneGraphAnnotations/Private/ZoneGraphAnnotationSubsystem.cpp b/Engine/Plugins/Runtime/ZoneGraphAnnotations/Source/ZoneGraphAnnotations/Private/ZoneGraphAnnotationSubsystem.cpp
--- a/Engine/Plugins/Runtime/ZoneGraphAnnotations/Source/ZoneGraphAnnotations/Private/ZoneGraphAnnotationSubsystem.cpp
+++ b/Engine/Plugins/Runtime/ZoneGraphAnnotations/Source/ZoneGraphAnnotations/Private/ZoneGraphAnnotationSubsystem.cpp
@@ -197,8 +197,3 @@ TStatId UZoneGraphAnnotationSubsystem::GetStatId() const
 {
 	RETURN_QUICK_DECLARE_CYCLE_STAT(UZoneGraphAnnotationSubsystem, STATGROUP_Tickables);
 }
-
-bool UZoneGraphAnnotationSubsystem::DoesSupportWorldType(const EWorldType::Type WorldType) const
-{
-	return WorldType == EWorldType::EditorStorage || Super::DoesSupportWorldType(WorldType);
-}
diff --git a/Engine/Plugins/Runtime/ZoneGraphAnnotations/Source/ZoneGraphAnnotations/Public/ZoneGraphAnnotationSubsystem.h b/Engine/Plugins/Runtime/ZoneGraphAnnotations/Source/ZoneGraphAnnotations/Public/ZoneGraphAnnotationSubsystem.h
--- a/Engine/Plugins/Runtime/ZoneGraphAnnotations/Source/ZoneGraphAnnotations/Public/ZoneGraphAnnotationSubsystem.h
+++ b/Engine/Plugins/Runtime/ZoneGraphAnnotations/Source/ZoneGraphAnnotations/Public/ZoneGraphAnnotationSubsystem.h
@@ -120,8 +120,6 @@ protected:
 	virtual void Tick(float DeltaTime) override;
 	virtual TStatId GetStatId() const override;
 
-	virtual bool DoesSupportWorldType(const EWorldType::Type WorldType) const override;
-
 	/** Array of registered components. */
 	UPROPERTY(Transient)
 	TArray<FRegisteredZoneGraphAnnotation> RegisteredComponents;
diff --git a/Engine/Source/Editor/UnrealEd/Private/EditorServer.cpp b/Engine/Source/Editor/UnrealEd/Private/EditorServer.cpp
--- a/Engine/Source/Editor/UnrealEd/Private/EditorServer.cpp
+++ b/Engine/Source/Editor/UnrealEd/Private/EditorServer.cpp
@@ -2015,8 +2015,7 @@ void UEditorEngine::CheckForWorldGCLeaks( UWorld* NewWorld, UPackage* WorldPacka
 		const bool bIsPersistantWorldType = (
 			RemainingWorld->WorldType == EWorldType::Inactive) || 
 			(RemainingWorld->WorldType == EWorldType::EditorPreview) || 
-			(RemainingWorld->WorldType == EWorldType::GamePreview) ||
-			(RemainingWorld->WorldType == EWorldType::EditorStorage);
+			(RemainingWorld->WorldType == EWorldType::GamePreview);
 		if(!bIsNewWorld && !bIsPersistantWorldType && !WorldHasValidContext(RemainingWorld))
 		{
 			FindAndPrintStaleReferencesToObject(RemainingWorld, EPrintStaleReferencesOptions::Error);
@@ -2070,9 +2069,7 @@ void UEditorEngine::EditorDestroyWorld( FWorldContext & Context, const FText& Cl
 		WorldPackage = NULL;
 	}
 
-	if (ContextWorld->WorldType != EWorldType::EditorPreview && 
-		ContextWorld->WorldType != EWorldType::EditorStorage &&
-		ContextWorld->WorldType != EWorldType::Inactive)
+	if (ContextWorld->WorldType != EWorldType::EditorPreview && ContextWorld->WorldType != EWorldType::Inactive)
 	{
 		// Go away, come again never!
 		ContextWorld->ClearFlags(RF_Standalone | RF_Transactional);
diff --git a/Engine/Source/Runtime/Engine/Classes/Engine/EngineTypes.h b/Engine/Source/Runtime/Engine/Classes/Engine/EngineTypes.h
--- a/Engine/Source/Runtime/Engine/Classes/Engine/EngineTypes.h
+++ b/Engine/Source/Runtime/Engine/Classes/Engine/EngineTypes.h
@@ -1117,13 +1117,6 @@ namespace EWorldType
 		/** A preview world for an editor tool */
 		EditorPreview,
 
-		/** 
-		 * A minimal world for storing editor information.
-		 * Alert: This world type is part of experimental development and it's strongly recommended to not use this world type
-		 *		as long a this alert is here.
-		 */
-		EditorStorage,
-
 		/** A preview world for a game */
 		GamePreview,
 
diff --git a/Engine/Source/Runtime/Engine/Classes/Engine/World.h b/Engine/Source/Runtime/Engine/Classes/Engine/World.h
--- a/Engine/Source/Runtime/Engine/Classes/Engine/World.h
+++ b/Engine/Source/Runtime/Engine/Classes/Engine/World.h
@@ -3810,13 +3810,6 @@ public:
 	/** Returns true if this world is any kind of editor world (including editor preview worlds) */
 	bool IsEditorWorld() const;
 
-	/** 
-	 * Returns true if this is a minimal world for storing information.
-	 * Alert: This is part of experimental development and it's strongly recommended to not use this function
-	 *		as long a this alert is here.
-	 */
-	bool IsStorageWorld() const;
-
 	/** Returns true if this world is a preview game world (editor or game) */
 	bool IsPreviewWorld() const;
 
diff --git a/Engine/Source/Runtime/Engine/Private/World.cpp b/Engine/Source/Runtime/Engine/Private/World.cpp
--- a/Engine/Source/Runtime/Engine/Private/World.cpp
+++ b/Engine/Source/Runtime/Engine/Private/World.cpp
@@ -7728,7 +7728,7 @@ bool UWorld::IsGameWorld() const
 
 bool UWorld::IsEditorWorld() const
 {
-	return WorldType == EWorldType::Editor || WorldType == EWorldType::EditorPreview || WorldType == EWorldType::PIE || WorldType == EWorldType::EditorStorage;
+	return WorldType == EWorldType::Editor || WorldType == EWorldType::EditorPreview || WorldType == EWorldType::PIE;
 }
 
 bool UWorld::IsPreviewWorld() const
@@ -7736,11 +7736,6 @@ bool UWorld::IsPreviewWorld() const
 	return WorldType == EWorldType::EditorPreview || WorldType == EWorldType::GamePreview;
 }
 
-bool UWorld::IsStorageWorld() const
-{
-	return WorldType == EWorldType::EditorStorage;
-}
-
 bool UWorld::UsesGameHiddenFlags() const
 {
 	return IsGameWorld();
@@ -9074,7 +9069,6 @@ FString ToString(EWorldType::Type Type)
 	case EWorldType::Editor: return TEXT("Editor");
 	case EWorldType::PIE: return TEXT("PIE");
 	case EWorldType::EditorPreview: return TEXT("EditorPreview");
-	case EWorldType::EditorStorage: return TEXT("EditorStorage");
 	case EWorldType::GamePreview: return TEXT("GamePreview");
 	case EWorldType::GameRPC: return TEXT("GameRPC");
 	case EWorldType::Inactive: return TEXT("Inactive");
