Make source actor filtering only keep weak references to filtered classes

diff --git a/Engine/Plugins/Developer/TraceSourceFiltering/Source/SourceFilteringTrace/Private/ActorFiltering.h b/Engine/Plugins/Developer/TraceSourceFiltering/Source/SourceFilteringTrace/Private/ActorFiltering.h
--- a/Engine/Plugins/Developer/TraceSourceFiltering/Source/SourceFilteringTrace/Private/ActorFiltering.h
+++ b/Engine/Plugins/Developer/TraceSourceFiltering/Source/SourceFilteringTrace/Private/ActorFiltering.h
@@ -90,7 +90,7 @@ private:
 };
 
 /** Used for retrieving AActor instances of particular classes from the owned UWorld, derived from FActorIteratorState */
-class FFilteredActorCollector : public FGCObject
+class FFilteredActorCollector
 {
 public:
 	FFilteredActorCollector(UWorld* InWorld, const TArray<FActorClassFilter>& FilteredClasses)
@@ -143,7 +143,7 @@ public:
 
 #if WITH_EDITOR
 		// Same behaviour as FActorIteratorState to reduce editor load
-		if (AllFilteredClasses.Num() == 1 && AllFilteredClasses[0] == AActor::StaticClass())
+		if (AllFilteredClasses.Num() == 1 && AllFilteredClasses[0].Get() == AActor::StaticClass())
 		{
 			// First determine the number of actors in the world to reduce reallocations when we append them to the array below.
 			int32 TotalNumActors = 0;
@@ -175,7 +175,7 @@ public:
 		EObjectFlags ExcludeFlags = RF_ClassDefaultObject;
 		EInternalObjectFlags InternalExcludeFlags = EInternalObjectFlags::Garbage;
 
-		ForEachObjectOfClasses(AllFilteredClasses, [this](UObject* Object)
+		ForEachObjectOfClasses(GetAllFilteredClasses(), [this](UObject* Object)
 		{
 			// Ensure the actor, level and world are valid
 			const AActor* Actor = static_cast<const AActor*>(Object);
@@ -204,18 +204,23 @@ public:
 		return ActorArray.Num();
 	}
 
-	void AddReferencedObjects(FReferenceCollector& Collector) override
+	TArray<const UClass*> GetAllFilteredClasses() const
 	{
-		Collector.AddStableReferenceArray(&AllFilteredClasses);
-	}
-	virtual FString GetReferencerName() const override
-	{
-		return TEXT("FFilteredActorCollector");
+		TArray<const UClass*> RawFilteredClasses;
+		RawFilteredClasses.Reserve(AllFilteredClasses.Num());
+		for(const TWeakObjectPtr<const UClass>& FilteredClass : AllFilteredClasses)
+		{
+			if (const UClass* RawFilteredClass = FilteredClass.Get())
+			{
+				RawFilteredClasses.Add(RawFilteredClass);
+			}
+		}
+		return RawFilteredClasses;
 	}
 
 protected:
 	UWorld* CurrentWorld;
 
 	TArray<TWeakObjectPtr<const AActor>> ActorArray;
-	TArray<TObjectPtr<const UClass>> AllFilteredClasses;
+	TArray<TWeakObjectPtr<const UClass>> AllFilteredClasses;
 };
