Enable AutoRTFM for this module

diff --git a/Engine/Plugins/Developer/TraceSourceFiltering/Source/SourceFilteringTrace/SourceFilteringTrace.Build.cs b/Engine/Plugins/Developer/TraceSourceFiltering/Source/SourceFilteringTrace/SourceFilteringTrace.Build.cs
--- a/Engine/Plugins/Developer/TraceSourceFiltering/Source/SourceFilteringTrace/SourceFilteringTrace.Build.cs
+++ b/Engine/Plugins/Developer/TraceSourceFiltering/Source/SourceFilteringTrace/SourceFilteringTrace.Build.cs
@@ -27,5 +27,7 @@ public class SourceFilteringTrace : ModuleRules
 				}
 			);
         }
+
+		bAllowAutoRTFMInstrumentation = true;
     }
 }
