Remove this that should have deadended in 30.10

diff --git a/Engine/Plugins/Developer/TraceSourceFiltering/Source/SourceFilteringTrace/SourceFilteringTrace.Build.cs b/Engine/Plugins/Developer/TraceSourceFiltering/Source/SourceFilteringTrace/SourceFilteringTrace.Build.cs
--- a/Engine/Plugins/Developer/TraceSourceFiltering/Source/SourceFilteringTrace/SourceFilteringTrace.Build.cs
+++ b/Engine/Plugins/Developer/TraceSourceFiltering/Source/SourceFilteringTrace/SourceFilteringTrace.Build.cs
@@ -27,7 +27,5 @@ public class SourceFilteringTrace : ModuleRules
 				}
 			);
         }
-
-		bAllowAutoRTFMInstrumentation = true;
     }
 }
